#SkillSystem
-GameplayEffect leveling
-Global curve tables
-Changes to gameplayeffect tagging

diff --git a/Engine/Source/Editor/SkillSystemEditor/Private/AttributeDetails.cpp b/Engine/Source/Editor/SkillSystemEditor/Private/AttributeDetails.cpp
--- a/Engine/Source/Editor/SkillSystemEditor/Private/AttributeDetails.cpp
+++ b/Engine/Source/Editor/SkillSystemEditor/Private/AttributeDetails.cpp
@@ -30,6 +30,8 @@ void FAttributePropertyDetails::CustomizeStructHeader( TSharedRef<IPropertyHandl
 	PropertyOptions.Empty();
 	PropertyOptions.Add(MakeShareable(new FString("None")));
 
+	FString FilterMetaStr = StructPropertyHandle->GetProperty()->GetMetaData(TEXT("FilterMetaTag"));
+
 	// Gather all UAttraibute classes
 	for (TObjectIterator<UClass> ClassIt; ClassIt; ++ClassIt)
 	{
@@ -39,6 +41,12 @@ void FAttributePropertyDetails::CustomizeStructHeader( TSharedRef<IPropertyHandl
 			for (TFieldIterator<UProperty> PropertyIt(Class, EFieldIteratorFlags::ExcludeSuper); PropertyIt; ++PropertyIt)
 			{
 				UProperty *Property = *PropertyIt;
+
+				if (!FilterMetaStr.IsEmpty() && Property->HasMetaData(*FilterMetaStr))
+				{
+					continue;
+				}
+				
 				PropertyOptions.Add(MakeShareable(new FString(FString::Printf(TEXT("%s.%s"), *Class->GetName(), *Property->GetName()))));
 			}
 		}
@@ -50,11 +58,22 @@ void FAttributePropertyDetails::CustomizeStructHeader( TSharedRef<IPropertyHandl
 			StructPropertyHandle->CreatePropertyNameWidget()
 		]
 		.ValueContent()
+		.MinDesiredWidth(500)
+		.MaxDesiredWidth(4096)
 		[
-			SNew(STextComboBox)
-			.OptionsSource( &PropertyOptions )
-			.InitiallySelectedItem(GetPropertyType())
-			.OnSelectionChanged( this, &FAttributePropertyDetails::OnChangeProperty )
+			SNew(SHorizontalBox)
+			+ SHorizontalBox::Slot()
+			//.FillWidth(1.0f)
+			.HAlign(HAlign_Fill)
+			.Padding(0.f, 0.f, 2.f, 0.f)
+			[
+
+				SNew(STextComboBox)
+				.ContentPadding(FMargin(2.0f, 2.0f))
+				.OptionsSource( &PropertyOptions )
+				.InitiallySelectedItem(GetPropertyType())
+				.OnSelectionChanged( this, &FAttributePropertyDetails::OnChangeProperty )
+			]
 		];
 }
 END_SLATE_FUNCTION_BUILD_OPTIMIZATION
diff --git a/Engine/Source/Runtime/SkillSystem/Classes/AttributeComponent.h b/Engine/Source/Runtime/SkillSystem/Classes/AttributeComponent.h
--- a/Engine/Source/Runtime/SkillSystem/Classes/AttributeComponent.h
+++ b/Engine/Source/Runtime/SkillSystem/Classes/AttributeComponent.h
@@ -26,6 +26,8 @@ class SKILLSYSTEM_API UAttributeComponent : public UActorComponent
 {
 	GENERATED_UCLASS_BODY()
 
+	virtual ~UAttributeComponent();
+
 	template <class T >
 	T*	GetSet()
 	{
@@ -63,9 +65,9 @@ class SKILLSYSTEM_API UAttributeComponent : public UActorComponent
 	FActiveGameplayEffectHandle ApplyGameplayEffectSpecToSelf(const FGameplayEffectSpec &GameplayEffect, FModifierQualifier BaseQualifier = FModifierQualifier());
 
 	UFUNCTION(BlueprintCallable, BlueprintAuthorityOnly, Category = GameplayEffects)
-	bool	RemoveActiveGameplayEffect(FActiveGameplayEffectHandle Handle);
+	bool RemoveActiveGameplayEffect(FActiveGameplayEffectHandle Handle);
 
-	float	GetGameplayEffectDuration(FActiveGameplayEffectHandle Handle) const;
+	float GetGameplayEffectDuration(FActiveGameplayEffectHandle Handle) const;
 
 	// Not happy with this interface but don't see a better way yet. How should outside code (UI, etc) ask things like 'how much is this gameplay effect modifying my damage by'
 	// (most likely we want to catch this on the backend - when damage is applied we can get a full dump/history of how the number got to where it is. But still we may need polling methods like below (how much would my damage be)
@@ -77,13 +79,13 @@ class SKILLSYSTEM_API UAttributeComponent : public UActorComponent
 
 	UFUNCTION(BlueprintCallable, BlueprintAuthorityOnly, Category = GameplayEffects)
 	bool IsGameplayEffectActive(FActiveGameplayEffectHandle InHandle) const;
-
+	
 
 	// --------------------------------------------
 	// Possibly useful but not primary API functions:
 	// --------------------------------------------
 
-	FActiveGameplayEffectHandle ApplyGameplayEffectSpecToTarget(UGameplayEffect *GameplayEffect, UAttributeComponent *Target, float Level = 1.f, FModifierQualifier BaseQualifier = FModifierQualifier()) const;
+	FActiveGameplayEffectHandle ApplyGameplayEffectSpecToTarget(UGameplayEffect *GameplayEffect, UAttributeComponent *Target, float Level = FGameplayEffectLevelSpec::INVALID_LEVEL, FModifierQualifier BaseQualifier = FModifierQualifier()) const;
 
 	UFUNCTION(BlueprintCallable, BlueprintAuthorityOnly, Category = GameplayEffects, meta=(FriendlyName = "ApplyGameplayEffectToSelf"))
 	FActiveGameplayEffectHandle K2_ApplyGameplayEffectToSelf(UGameplayEffect *GameplayEffect, float Level, AActor *Instigator);
@@ -92,12 +94,12 @@ class SKILLSYSTEM_API UAttributeComponent : public UActorComponent
 
 	int32 GetNumActiveGameplayEffect() const;
 
+	void AddDependancyToAttribute(FGameplayAttribute Attribute, const TWeakPtr<FAggregator> InDependant);
+
 	// --------------------------------------------
 	// Temp / Debug
 	// --------------------------------------------
 
-	void RegisterAttributeModifyCallback(FGameplayAttribute Attribute, FOnGameplayAttributeEffectExecuted Delegate);
-
 	void TEMP_ApplyActiveGameplayEffects();
 	
 	void PrintAllGameplayEffects() const;
diff --git a/Engine/Source/Runtime/SkillSystem/Classes/GameplayEffect.h b/Engine/Source/Runtime/SkillSystem/Classes/GameplayEffect.h
--- a/Engine/Source/Runtime/SkillSystem/Classes/GameplayEffect.h
+++ b/Engine/Source/Runtime/SkillSystem/Classes/GameplayEffect.h
@@ -26,11 +26,14 @@ UENUM(BlueprintType)
 namespace EGameplayModOp
 {
 	enum Type
-	{
-		Override = 0		UMETA(DisplayName="Override"),
-		Additive			UMETA(DisplayName="Add"),
+	{		
+		// Numeric
+		Additive = 0		UMETA(DisplayName="Add"),
 		Multiplicitive		UMETA(DisplayName="Multiply"),
 		Division			UMETA(DisplayName="Divide"),
+
+		// Other
+		Override 			UMETA(DisplayName="Override"),	// This should always be the first non numeric ModOp
 		Callback			UMETA(DisplayName="Custom"),
 
 		// This must always be at the end
@@ -107,6 +110,30 @@ struct FGameplayModifierCallbacks
 	TArray<TSubclassOf<class UGameplayEffectExtension> >	ExtensionClasses;
 };
 
+/**
+* Defines how A GameplayEffect levels
+*	Normally, GameplayEffect levels are specified when they are created.
+*	They can also be tied to their instigators attribute.
+*		For example, a Damage apply GameplayEffect that 'levels' based on the PhysicalDamage attribute
+*/
+USTRUCT()
+struct SKILLSYSTEM_API FGameplayEffectLevelDef
+{
+	GENERATED_USTRUCT_BODY()
+
+	/** When true, whatever creates of owns this will pass in a level. E.g, level is not intrinsic to this definition. */
+	UPROPERTY(EditDefaultsOnly, Category = GameplayEffectLevel)
+	bool InheritLevelFromOwner;
+
+	/** If set, the gameplay effects level will be tied to this attribute on the instigator */
+	UPROPERTY(EditDefaultsOnly, Category = GameplayEffectLevel, meta = (FilterMetaTag="HideFromLevelInfos"))
+	FGameplayAttribute	Attribute;
+
+	/** If true, take snapshot of attribute level when the gameplay effect is initialized. Otherwise, the level of the gameplay effect will update as the attribute it is tied to updates */
+	UPROPERTY(EditDefaultsOnly, Category = GameplayEffectLevel)
+	bool TakeSnapshotOnInit;
+};
+
 /**
  * FGameplayModifierInfo
  *	Tells us "Who/What we" modify
@@ -152,14 +179,14 @@ struct SKILLSYSTEM_API FGameplayModifierInfo
 	// The thing I modify must not have any of these tags
 	UPROPERTY(EditDefaultsOnly, Category = GameplayModifier)
 	FGameplayTagContainer IgnoreTags;
-
-	// These tags I pass on to whatever I modify
-	UPROPERTY(EditDefaultsOnly, Category=GameplayModifier)
-	FGameplayTagContainer PassedTags;
-
+	
+	/** This modifiers tags. These tags are passed to any other modifiers that this modifies. */
 	UPROPERTY(EditDefaultsOnly, Category=GameplayModifier)
 	FGameplayTagContainer OwnedTags;
 
+	UPROPERTY(EditDefaultsOnly, BlueprintReadWrite, Category = GameplayEffect)
+	FGameplayEffectLevelDef	LevelInfo;
+
 	FString ToSimpleString() const
 	{
 		return FString::Printf(TEXT("%s %s %s BaseVaue: %s"), *EGameplayModToString(ModifierType), *EGameplayModOpToString(ModifierOp), *EGameplayModEffectToString(EffectType), *Magnitude.ToSimpleString());
@@ -232,7 +259,7 @@ public:
 	/** Constant specifying that the combat effect has no period and doesn't check for over time application */
 	static const float NO_PERIOD;
 	
-	UPROPERTY(EditDefaultsOnly, Category=GameplayEffect, meta=(GameplayAttribute="True"))
+	UPROPERTY(EditDefaultsOnly, Category=GameplayEffect)
 	FScalableFloat	Duration;
 
 	UPROPERTY(EditDefaultsOnly, Category=GameplayEffect)
@@ -241,28 +268,26 @@ public:
 	UPROPERTY(EditDefaultsOnly, BlueprintReadOnly, Category=GameplayEffect)
 	TArray<FGameplayModifierInfo> Modifiers;
 
-	/** Description of this combat effect. */
-	UPROPERTY(EditDefaultsOnly, BlueprintReadWrite, Category = Display)
-	FText Description;
-
+	UPROPERTY(EditDefaultsOnly, BlueprintReadWrite, Category=GameplayEffect)
+	FGameplayEffectLevelDef	LevelInfo;
+		
 	// "I can only be applied to targets that have these tags"
 	// "I can only exist on CE buckets on targets that have these tags":
 	
 	/** Container of gameplay tags that have to be present on the target for the effect to be applied */
-	UPROPERTY(EditDefaultsOnly, BlueprintReadOnly, Category=GameplayEffect)
+	UPROPERTY(EditDefaultsOnly, BlueprintReadOnly, Category = Application)
 	FGameplayTagContainer ApplicationRequiredTargetTags;
 
-
 	// "I can only be applied if my instigator has these tags"
 
 	/** Container of gameplay tags that have to be present on the instigator for the effect to be applied */
-	UPROPERTY(EditDefaultsOnly, BlueprintReadOnly, Category=ApplicationInstigatorTags)
+	UPROPERTY(EditDefaultsOnly, BlueprintReadOnly, Category = Application)
 	FGameplayTagContainer ApplicationRequiredInstigatorTags;
 
-	UPROPERTY(EditDefaultsOnly, Category=GameplayEffect, meta=(GameplayAttribute="True"))
+	UPROPERTY(EditDefaultsOnly, Category=Application, meta=(GameplayAttribute="True"))
 	FScalableFloat	ChanceToApplyToTarget;
 
-	UPROPERTY(EditDefaultsOnly, Category=GameplayEffect, meta=(GameplayAttribute="True"))
+	UPROPERTY(EditDefaultsOnly, Category = Application, meta = (GameplayAttribute = "True"))
 	FScalableFloat	ChanceToApplyToGameplayEffect;
 
 	// Modify duration of CEs
@@ -290,18 +315,18 @@ public:
 	FGameplayTagContainer GameplayEffectTags;
 
 	// "In order to affect another GE, they must have ALL of these tags"
-	UPROPERTY(EditDefaultsOnly, BlueprintReadOnly, Category=GameplayEffect)
+	UPROPERTY(EditDefaultsOnly, BlueprintReadOnly, Category=Tags)
 	FGameplayTagContainer GameplayEffectRequiredTags;
 	
 	// "In order to affect another GE, they must NOT have ANY of these tags"
-	UPROPERTY(EditDefaultsOnly, BlueprintReadOnly, Category=GameplayEffect)
+	UPROPERTY(EditDefaultsOnly, BlueprintReadOnly, Category=Tags)
 	FGameplayTagContainer GameplayEffectIgnoreTags;
 
 	/** Can this GameplayEffect modify the input parameter, based on tags  */
 	bool AreGameplayEffectTagRequirementsSatisfied(const UGameplayEffect *GameplayEffectToBeModified) const
 	{
-		bool HasRequired = GameplayEffectToBeModified->OwnedTagsContainer.HasAllTags(GameplayEffectRequiredTags);
-		bool HasIgnored = GameplayEffectToBeModified->OwnedTagsContainer.Num() > 0 && GameplayEffectToBeModified->OwnedTagsContainer.HasAnyTag(GameplayEffectIgnoreTags);
+		bool HasRequired = GameplayEffectToBeModified->GameplayEffectTags.HasAllTags(GameplayEffectRequiredTags);
+		bool HasIgnored = GameplayEffectIgnoreTags.Num() > 0 && GameplayEffectToBeModified->GameplayEffectTags.HasAnyTag(GameplayEffectIgnoreTags);
 
 		return HasRequired && !HasIgnored;
 	}
@@ -325,11 +350,11 @@ public:
 	}
 
 	/** Container of "owned" gameplay tags that are applied to actor with this effect applied to them */
-	UPROPERTY(EditDefaultsOnly, BlueprintReadOnly, Category=OwnedTags)
+	UPROPERTY(EditDefaultsOnly, BlueprintReadOnly, Category=Tags)
 	FGameplayTagContainer OwnedTagsContainer;
 
 	/** Container of gameplay tags to be cleared upon effect application; Any active effects with these tags that can be cleared, will be */
-	UPROPERTY(EditDefaultsOnly, BlueprintReadOnly, Category=ClearTags)
+	UPROPERTY(EditDefaultsOnly, BlueprintReadOnly, Category=Tags)
 	FGameplayTagContainer ClearTagsContainer;
 	
 	// Used to quickly tell if a GameplayEffect modifies another GameplayEffect (or a set of attributes)
@@ -340,13 +365,17 @@ public:
 
 	// -----------------------------------------------
 	
-	UPROPERTY(EditDefaultsOnly, BlueprintReadOnly, Category = GameplayEffect)
+	UPROPERTY(EditDefaultsOnly, BlueprintReadOnly, Category = Advanced)
 	TEnumAsByte<EGameplayEffectCopyPolicy::Type>	CopyPolicy;
 
 	// ----------------------------------------------
 
-	UPROPERTY(EditDefaultsOnly, BlueprintReadOnly, Category = GameplayEffect)
+	UPROPERTY(EditDefaultsOnly, BlueprintReadOnly, Category = Display)
 	TArray<FGameplayEffectCue>	GameplayCues;
+
+	/** Description of this combat effect. */
+	UPROPERTY(EditDefaultsOnly, BlueprintReadWrite, Category = Display)
+	FText Description;
 };
 
 /**
@@ -404,50 +433,83 @@ private:
 
 /**
  * FGameplayEffectLevelSpec
- *	Level Specification. For now, this is just a flat ConstantLevel. We would like to support dynamic level spec: e.g, tieing a Level to a callback delegate.
- *	This could be used to tie a GameplayEffects power level to an attribute value, etc. 
+ *	Level Specification. This can be a static, constant level specified on creation or it can be dynamically tied to a source's attribute value.
+ *	For example, a GameplayEffect could be made whose level is tied to its instigators PhysicalDamage or Intelligence attribute.
  */
-USTRUCT()
+
 struct FGameplayEffectLevelSpec
 {
-	GENERATED_USTRUCT_BODY()
-
 	static const float INVALID_LEVEL;
 
 	FGameplayEffectLevelSpec()
 		:ConstantLevel(INVALID_LEVEL)
+		,CachedLevel(INVALID_LEVEL)
 	{
 	}
 
-	FGameplayEffectLevelSpec(float InLevel)
+	FGameplayEffectLevelSpec(float InLevel, const FGameplayEffectLevelDef &Def, class AActor *InSource)
 		: ConstantLevel(InLevel)
+		, CachedLevel(InLevel)
+		, Source(InSource)
 	{
+		if (Def.Attribute.GetUProperty() != NULL)
+		{
+			Attribute = Def.Attribute;
+		}
+
+		if (Def.TakeSnapshotOnInit)
+		{
+			SnapshotLevel();
+		}
 	}
 
-	bool IsValid() const
+	void ApplyNewDef(const FGameplayEffectLevelDef &Def, TSharedPtr<FGameplayEffectLevelSpec> &OutSharedPtr) const
 	{
-		return (ConstantLevel != INVALID_LEVEL);
+		if (Def.InheritLevelFromOwner)
+		{
+			return;
+		}
+
+		check(OutSharedPtr.IsValid());
+		if (Def.Attribute != OutSharedPtr->Attribute)
+		{
+			// In def levels off something different
+			// make a new level spec
+			OutSharedPtr = TSharedPtr<FGameplayEffectLevelSpec>(new FGameplayEffectLevelSpec(INVALID_LEVEL, Def, Source.Get()));
+		}
 	}
 
-	float GetLevel() const
+	/** Dynamic simply means the level may change. It is not constant. */
+	bool IsDynamic() const
+	{
+		return ConstantLevel == INVALID_LEVEL && Attribute.GetUProperty() != NULL;
+	}
+
+	/** Valid means we have some meaningful data. If we have an INVALID_LEVEL constant value and are not tied to a dynamic property, then we are invalid. */
+	bool IsValid() const
 	{
-		check(IsValid());
-		return ConstantLevel;
+		return ConstantLevel != INVALID_LEVEL || Attribute.GetUProperty() != NULL;
 	}
 
+	float GetLevel() const;
+
 	void SnapshotLevel()
 	{
 		// This should snapshot the current level (if dynamic/delegate) and save off its value so that it doesn't change
+		ConstantLevel = GetLevel();
+		Source = NULL;
 	}
 
+	void RegisterLevelDependancy(TWeakPtr<struct FAggregator> OwningAggregator);
+
 	void PrintAll() const;
 
-	UPROPERTY(EditDefaultsOnly, Category=ComabtEffectLevel)
-	float	ConstantLevel;
+	mutable float ConstantLevel;	// Final/constant level. Once this is set we are locked at the given level.
+	mutable float CachedLevel;		// Last read value. Needed in case we lose our source, we use the last known level.
 
-	// Not supported yet, but this will be awesome!
-	// UPROPERTY(EditDefaultsOnly, Category=ComabtEffectLevel)
-	// FAttributeProperty	Property;
+	TWeakObjectPtr<class AActor>	Source;
+
+	FGameplayAttribute	Attribute;
 };
 
 /**
@@ -598,6 +660,11 @@ struct FGameplayModifierData
 	{
 		Magnitude = Info.Magnitude.MakeFinalizedCopy(CurveData);
 		Tags = Info.OwnedTags;
+
+		// Fixme: this is static data, should be a reference
+		RequireTags = Info.RequiredTags;
+		IgnoreTags = Info.IgnoreTags;
+
 		if (Info.Callbacks.ExtensionClasses.Num() > 0)
 		{
 			Callbacks = &Info.Callbacks;
@@ -611,15 +678,11 @@ struct FGameplayModifierData
 		Callbacks = NULL;
 	}
 
-	FGameplayModifierData(float InMagnitude, const FGameplayModifierCallbacks * InCallbacks, const FGameplayTagContainer *InTags = NULL)
+	FGameplayModifierData(float InMagnitude, const FGameplayModifierCallbacks * InCallbacks)
 	{
 		// Magnitude will be fixed at this value
 		Magnitude.SetValue(InMagnitude);
 		Callbacks = InCallbacks;
-		if (InTags)
-		{
-			Tags = *InTags;
-		}
 	}
 
 	// That magnitude that we modify by
@@ -627,6 +690,9 @@ struct FGameplayModifierData
 
 	// The tags I have
 	FGameplayTagContainer Tags;
+	
+	FGameplayTagContainer RequireTags;
+	FGameplayTagContainer IgnoreTags;
 
 	// Callback information for custom logic pre/post evaluation
 	const FGameplayModifierCallbacks * Callbacks;
@@ -700,6 +766,7 @@ struct FAggregator : public TSharedFromThis<FAggregator>
 	virtual ~FAggregator();
 
 	FAggregator & MarkDirty();
+	void CleaerAllDependancies();
 
 	const FGameplayModifierEvaluatedData& Evaluate() const;
 
@@ -714,6 +781,8 @@ struct FAggregator : public TSharedFromThis<FAggregator>
 
 	void AddDependantAggregator(TWeakPtr<FAggregator> InDependant);
 
+	void RegisterLevelDependancies();
+
 	TSharedPtr<FGameplayEffectLevelSpec> Level;
 	FActiveGameplayEffectHandle	ActiveHandle;	// Handle to owning active effect. Will be null in many cases.
 
@@ -887,7 +956,7 @@ private:
 };
 
 /**
- * Modifier Speciication
+ * Modifier Specification
  *	-Const data (FGameplayModifierInfo) tells us what we modify, what we can modify
  *	-Mutable Aggregated data tells us how we modify (magnitude).
  *  
@@ -900,6 +969,7 @@ struct FModifierSpec
 		: Info(InInfo)
 		, Aggregator(new FAggregator(FGameplayModifierData(InInfo, CurveData), InLevel,  SKILL_AGG_DEBUG(TEXT("FModifierSpec: %s "), *InInfo.ToSimpleString() )))
 	{
+		Aggregator.Get()->RegisterLevelDependancies();
 	}
 
 	// Hard Ref to what we modify, this stuff is const and never changes
@@ -910,8 +980,6 @@ struct FModifierSpec
 	bool CanModifyInContext(const FModifierQualifier &QualifierContext) const;
 	// returns true if this GameplayEffect can modify Other, false otherwise
 	bool CanModifyModifier(FModifierSpec &Other, const FModifierQualifier &QualifierContext) const;
-	// returns true if this GameplayEffect can modify Ref, false otherwise
-	bool CanModifyAggregator(FAggregatorRef Ref, const FModifierQualifier &QualifierContext) const;
 	
 	void ApplyModTo(FModifierSpec &Other, bool TakeSnapshot) const;
 	void ExecuteModOn(FModifierSpec &Other) const;
@@ -925,10 +993,6 @@ struct FModifierSpec
 	// Returns true if it can modify the input parameter, false otherwise
 	bool AreTagRequirementsSatisfied(const FModifierSpec &ModifierToBeModified) const;
 
-	// Can this GameplayEffect modify the input parameter, based on tags
-	// Returns true if it can modify the input parameter, false otherwise
-	bool AreTagRequirementsSatisfied(const FAggregatorRef &AggregatorToBeModified) const;
-
 	void PrintAll() const;
 };
 
@@ -954,17 +1018,10 @@ struct FGameplayEffectSpec
 		// If we initialize a GameplayEffectSpec with no level object passed in.
 	}
 
-	FGameplayEffectSpec( UGameplayEffect * InDef, TSharedPtr<FGameplayEffectLevelSpec> InLevel, const FGlobalCurveDataOverride *CurveData )
-		: Def(InDef)
-		, ModifierLevel(InLevel)
-		, Duration(new FAggregator(InDef->Duration.MakeFinalizedCopy(CurveData), InLevel, SKILL_AGG_DEBUG(TEXT("%s Duration"), *InDef->GetName())))
-		, Period(new FAggregator(InDef->Period.MakeFinalizedCopy(CurveData), InLevel, SKILL_AGG_DEBUG(TEXT("%s Period"), *InDef->GetName())))
-	{
-		InitModifiers(CurveData);
-	}
+	FGameplayEffectSpec( const UGameplayEffect * InDef, TSharedPtr<FGameplayEffectLevelSpec> InLevel, const FGlobalCurveDataOverride *CurveData );
 	
 	UPROPERTY()
-	UGameplayEffect * Def;
+	const UGameplayEffect * Def;
 
 	// Replicated	
 	TSharedPtr< FGameplayEffectLevelSpec > ModifierLevel;
@@ -990,8 +1047,6 @@ struct FGameplayEffectSpec
 	int32 ApplyModifiersFrom(const FGameplayEffectSpec &InSpec, const FModifierQualifier &QualifierContext);
 	int32 ExecuteModifiersFrom(const FGameplayEffectSpec &InSpec, const FModifierQualifier &QualifierContext);
 
-	int32 ApplyModifiersToAggregators(const FModifierSpec &InMod, FAggregatorRef &MyMod, const FModifierQualifier &QualifierContext, bool ShouldSnapshot);
-
 	bool ShouldApplyAsSnapshot(const FModifierQualifier &QualifierContext) const;
 
 	FString ToSimpleString() const
@@ -1103,6 +1158,8 @@ struct FActiveGameplayEffectsContainer : public FFastArraySerializer
 	
 	bool ExecuteGameplayEffect(FActiveGameplayEffectHandle Handle);	// This should not be outward facing to the skill system API, should only be called by the owning attribute component
 
+	void AddDependancyToAttribute(FGameplayAttribute Attribute, const TWeakPtr<FAggregator> InDependant);
+
 	bool RemoveActiveGameplayEffect(FActiveGameplayEffectHandle Handle);
 	
 	float GetGameplayEffectDuration(FActiveGameplayEffectHandle Handle) const;
@@ -1111,8 +1168,6 @@ struct FActiveGameplayEffectsContainer : public FFastArraySerializer
 
 	bool IsGameplayEffectActive(FActiveGameplayEffectHandle Handle) const;
 
-	void RegisterAttributeModifyCallback(FGameplayAttribute Attribute, FOnGameplayAttributeEffectExecuted Delegate);
-
 	void PrintAllGameplayEffects() const;
 
 	int32 GetNumGameplayEffects() const
@@ -1133,16 +1188,16 @@ struct FActiveGameplayEffectsContainer : public FFastArraySerializer
 		return FastArrayDeltaSerialize<FActiveGameplayEffect>(GameplayEffects, DeltaParms, *this);
 	}
 
+	void PreDestroy();
+
 private:
 
+	FAggregatorRef & FindOrCreateAttributeAggregator(FGameplayAttribute Attribute);
+
 	FActiveGameplayEffectHandle LastAssignedHandle;
 
 	TMap<FGameplayAttribute, FAggregatorRef> OngoingPropertyEffects;
 
-	// Fixme: at the least, these two maps should be combined and a new structure should be made.
-	// alternatively it may make sence to put the 'pre execute' delegate at the FAggregator level instead.
-	TMap<FGameplayAttribute, FOnGameplayAttributeEffectExecuted> PropertyExecutionCallbacks;
-
 	bool IsNetAuthority() const;
 };
 
diff --git a/Engine/Source/Runtime/SkillSystem/Classes/SkillSystemTestAttributeSet.h b/Engine/Source/Runtime/SkillSystem/Classes/SkillSystemTestAttributeSet.h
--- a/Engine/Source/Runtime/SkillSystem/Classes/SkillSystemTestAttributeSet.h
+++ b/Engine/Source/Runtime/SkillSystem/Classes/SkillSystemTestAttributeSet.h
@@ -8,12 +8,24 @@ class USkillSystemTestAttributeSet : public UAttributeSet
 {
 	GENERATED_UCLASS_BODY()
 
-	UPROPERTY(EditAnywhere, BlueprintReadWrite, Category = "AttributeTest")
+	UPROPERTY(EditAnywhere, BlueprintReadWrite, Category = "AttributeTest", meta = (HideFromModifiers))			// You can't make a GameplayEffect modifiy Health Directly (go through)
+	float	MaxHealth;
+
+	UPROPERTY(EditAnywhere, BlueprintReadWrite, Category = "AttributeTest", meta = (HideFromModifiers))			// You can't make a GameplayEffect modifiy Health Directly (go through)
 	float	Health;
 
-	UPROPERTY(EditAnywhere, BlueprintReadWrite, Category = "AttributeTest")
+	/** This Damage is just used for applying negative health mods. Its not a 'persistent' attribute. */
+	UPROPERTY(EditAnywhere, BlueprintReadWrite, Category = "AttributeTest", meta = (HideFromLevelInfos))		// You can't make a GameplayEffect 'powered' by Damage (Its transient)
 	float	Damage;
 
+	/** This Attribute is the actual spell damage for this actor. It will power spell based GameplayEffects */
+	UPROPERTY(EditAnywhere, BlueprintReadWrite, Category = "AttributeTest")
+	float	SpellDamage;
+
+	/** This Attribute is the actual physical damage for this actor. It will power spell based GameplayEffects */
+	UPROPERTY(EditAnywhere, BlueprintReadWrite, Category = "AttributeTest")
+	float	PhysicalDamage;
+
 	UPROPERTY(EditAnywhere, BlueprintReadWrite, Category = "AttributeTest")
 	float	CritChance;
 
@@ -29,6 +41,9 @@ class USkillSystemTestAttributeSet : public UAttributeSet
 	UPROPERTY(EditAnywhere, BlueprintReadWrite, Category = "AttributeTest")
 	float	LifeSteal;
 
+	UPROPERTY(EditAnywhere, BlueprintReadWrite, Category = "AttributeTest")
+	float	Strength;
+
 	virtual void PreAttributeModify(struct FGameplayEffectModCallbackData &Data) OVERRIDE;
 	virtual void PostAttributeModify(const struct FGameplayEffectModCallbackData &Data) OVERRIDE;
 };
\ No newline at end of file
diff --git a/Engine/Source/Runtime/SkillSystem/Classes/SkillSystemTestPawn.h b/Engine/Source/Runtime/SkillSystem/Classes/SkillSystemTestPawn.h
--- a/Engine/Source/Runtime/SkillSystem/Classes/SkillSystemTestPawn.h
+++ b/Engine/Source/Runtime/SkillSystem/Classes/SkillSystemTestPawn.h
@@ -31,9 +31,4 @@ class SKILLSYSTEM_API ASkillSystemTestPawn : public ADefaultPawn, public IGamepl
 	FGameplayCueHandler	GameplayCueHandler;
 
 	static FName AttributeComponentName;
-
-
-	void OnDamageExecute(struct FGameplayModifierEvaluatedData & Data);
-
-	void OnDidDamage(const struct FGameplayEffectSpec & Spec, struct FGameplayModifierEvaluatedData & Data);
 };
\ No newline at end of file
diff --git a/Engine/Source/Runtime/SkillSystem/Private/AttributeComponent.cpp b/Engine/Source/Runtime/SkillSystem/Private/AttributeComponent.cpp
--- a/Engine/Source/Runtime/SkillSystem/Private/AttributeComponent.cpp
+++ b/Engine/Source/Runtime/SkillSystem/Private/AttributeComponent.cpp
@@ -25,6 +25,11 @@ UAttributeComponent::UAttributeComponent(const class FPostConstructInitializePro
 	ActiveGameplayEffects.Owner = this;
 }
 
+UAttributeComponent::~UAttributeComponent()
+{
+	ActiveGameplayEffects.PreDestroy();
+}
+
 UAttributeSet* UAttributeComponent::InitStats(TSubclassOf<class UAttributeSet> Attributes, const UDataTable* DataTable)
 {
 	UAttributeSet* AttributeObj = NULL;
@@ -164,7 +169,7 @@ FActiveGameplayEffectHandle UAttributeComponent::ApplyGameplayEffectSpecToTarget
 {
 	check(GameplayEffect);
 
-	FGameplayEffectSpec	Spec(GameplayEffect, TSharedPtr<FGameplayEffectLevelSpec>(new FGameplayEffectLevelSpec(Level)), GetCurveDataOverride());
+	FGameplayEffectSpec	Spec(GameplayEffect, TSharedPtr<FGameplayEffectLevelSpec>(new FGameplayEffectLevelSpec(Level, GameplayEffect->LevelInfo, GetOwner())), GetCurveDataOverride());
 	Spec.Def = GameplayEffect;
 	Spec.InstigatorStack.AddInstigator(GetOwner());
 	
@@ -180,7 +185,7 @@ FActiveGameplayEffectHandle UAttributeComponent::K2_ApplyGameplayEffectToSelf(UG
 /** This is a helper function - it seems like this will be useful as a blueprint interface at the least, but Level parameter may need to be expanded */
 FActiveGameplayEffectHandle UAttributeComponent::ApplyGameplayEffectToSelf(UGameplayEffect *GameplayEffect, float Level, AActor *Instigator, FModifierQualifier BaseQualifier)
 {
-	FGameplayEffectSpec	Spec(GameplayEffect, TSharedPtr<FGameplayEffectLevelSpec>(new FGameplayEffectLevelSpec(Level)), GetCurveDataOverride());
+	FGameplayEffectSpec	Spec(GameplayEffect, TSharedPtr<FGameplayEffectLevelSpec>(new FGameplayEffectLevelSpec(Level, GameplayEffect->LevelInfo, GetOwner())), GetCurveDataOverride());
 	Spec.Def = GameplayEffect;
 	Spec.InstigatorStack.AddInstigator(Instigator);
 
@@ -200,7 +205,6 @@ int32 UAttributeComponent::GetNumActiveGameplayEffect() const
 bool UAttributeComponent::IsGameplayEffectActive(FActiveGameplayEffectHandle InHandle) const
 {
 	return ActiveGameplayEffects.IsGameplayEffectActive(InHandle);
-
 }
 
 void UAttributeComponent::TEMP_ApplyActiveGameplayEffects()
@@ -345,10 +349,6 @@ float UAttributeComponent::GetGameplayEffectMagnitude(FActiveGameplayEffectHandl
 {
 	return ActiveGameplayEffects.GetGameplayEffectMagnitude(Handle, Attribute);
 }
-void UAttributeComponent::RegisterAttributeModifyCallback(FGameplayAttribute Attribute, FOnGameplayAttributeEffectExecuted Delegate)
-{
-	return ActiveGameplayEffects.RegisterAttributeModifyCallback(Attribute, Delegate);
-}
 
 void UAttributeComponent::InvokeGameplayCueExecute(const FGameplayEffectSpec &Spec)
 {
@@ -438,6 +438,11 @@ void UAttributeComponent::OnRep_GameplayEffects()
 
 }
 
+void UAttributeComponent::AddDependancyToAttribute(FGameplayAttribute Attribute, const TWeakPtr<FAggregator> InDependant)
+{
+	ActiveGameplayEffects.AddDependancyToAttribute(Attribute, InDependant);
+}
+
 // ---------------------------------------------------------------------------------------
 
 void UAttributeComponent::PrintAllGameplayEffects() const
@@ -490,7 +495,6 @@ void FModifierSpec::PrintAll() const
 	SKILL_LOG(Log, TEXT("ModifierOp: %s"), *EGameplayModOpToString(Info.ModifierOp));
 	SKILL_LOG(Log, TEXT("EffectType: %s"), *EGameplayModEffectToString(Info.EffectType));
 	SKILL_LOG(Log, TEXT("RequiredTags: %s"), *Info.RequiredTags.ToString());
-	SKILL_LOG(Log, TEXT("PassedTags: %s"), *Info.PassedTags.ToString());
 	SKILL_LOG(Log, TEXT("OwnedTags: %s"), *Info.OwnedTags.ToString());
 	SKILL_LOG(Log, TEXT("(Base) Magnitude: %s"), *Info.Magnitude.ToSimpleString());
 
diff --git a/Engine/Source/Runtime/SkillSystem/Private/GameplayEffect.cpp b/Engine/Source/Runtime/SkillSystem/Private/GameplayEffect.cpp
--- a/Engine/Source/Runtime/SkillSystem/Private/GameplayEffect.cpp
+++ b/Engine/Source/Runtime/SkillSystem/Private/GameplayEffect.cpp
@@ -56,6 +56,18 @@ bool UGameplayEffect::AreApplicationTagRequirementsSatisfied(const TSet<FName>&
 //
 // --------------------------------------------------------------------------------------------------------------------------------------------------------
 
+FGameplayEffectSpec::FGameplayEffectSpec(const UGameplayEffect * InDef, TSharedPtr<FGameplayEffectLevelSpec> InLevel, const FGlobalCurveDataOverride *CurveData)
+	: Def(InDef)
+	, ModifierLevel(InLevel)
+	, Duration(new FAggregator(InDef->Duration.MakeFinalizedCopy(CurveData), InLevel, SKILL_AGG_DEBUG(TEXT("%s Duration"), *InDef->GetName())))
+	, Period(new FAggregator(InDef->Period.MakeFinalizedCopy(CurveData), InLevel, SKILL_AGG_DEBUG(TEXT("%s Period"), *InDef->GetName())))
+{
+	Duration.Get()->RegisterLevelDependancies();
+	Period.Get()->RegisterLevelDependancies();
+
+	InitModifiers(CurveData);
+}
+
 void FGameplayEffectSpec::InitModifiers(const FGlobalCurveDataOverride *CurveData)
 {
 	check(Def);
@@ -64,9 +76,16 @@ void FGameplayEffectSpec::InitModifiers(const FGlobalCurveDataOverride *CurveDat
 	
 	for (const FGameplayModifierInfo &ModInfo : Def->Modifiers)
 	{
+		// Pass down the LevelInfo into this Modifier. 
+		// ModifierSpec may want to override how leveling will work (for this modifier only).
+		// Or it may use the GameplayEffectSpec's level. 
+		// ApplyNewDef will update NewLevelSpec appropriately.
+
+		TSharedPtr<FGameplayEffectLevelSpec> NewLevelSpec = ModifierLevel;
+		ModifierLevel->ApplyNewDef(ModInfo.LevelInfo, NewLevelSpec);
+
 		// This creates a new FModifierSpec that we own.
-		// The modifier level instance is shared through a shared pointer though
-		Modifiers.Emplace( FModifierSpec(ModInfo, ModifierLevel, CurveData));	
+		Modifiers.Emplace(FModifierSpec(ModInfo, NewLevelSpec, CurveData));
 	}	
 }
 
@@ -80,6 +99,8 @@ void FGameplayEffectSpec::MakeUnique()
 
 int32 FGameplayEffectSpec::ApplyModifiersFrom(const FGameplayEffectSpec &InSpec, const FModifierQualifier &QualifierContext)
 {
+	SKILL_LOG_SCOPE(TEXT("FGameplayEffectSpec::ApplyModifiersFrom %s. InSpec: %s"), *this->ToSimpleString(), *InSpec.ToSimpleString());
+
 	int32 NumApplied = 0;
 	bool ShouldSnapshot = InSpec.ShouldApplyAsSnapshot(QualifierContext);
 
@@ -89,6 +110,8 @@ int32 FGameplayEffectSpec::ApplyModifiersFrom(const FGameplayEffectSpec &InSpec,
 	}
 
 	// Fixme: Use acceleration structures to speed up these types of lookups
+	// The called functions below are reliant on the InSpecs evaluated data. We should ideally only call evaluate once per mod.
+
 	for (const FModifierSpec &InMod : InSpec.Modifiers)
 	{
 		if (!InMod.CanModifyInContext(QualifierContext))
@@ -105,13 +128,8 @@ int32 FGameplayEffectSpec::ApplyModifiersFrom(const FGameplayEffectSpec &InSpec,
 					if (InMod.CanModifyModifier(MyMod, QualifierContext))
 					{
 						InMod.ApplyModTo(MyMod, ShouldSnapshot);
-						Def->OwnedTagsContainer.AppendTags(InMod.Info.PassedTags);
 						NumApplied++;
 					}
-					else
-					{
-						NumApplied += ApplyModifiersToAggregators(InMod, MyMod.Aggregator, QualifierContext, ShouldSnapshot);
-					}
 				}
 				break;
 			}
@@ -123,7 +141,6 @@ int32 FGameplayEffectSpec::ApplyModifiersFrom(const FGameplayEffectSpec &InSpec,
 			case EGameplayModEffect::Duration:
 			{
 				Duration.Get()->ApplyMod(InMod.Info.ModifierOp, InMod.Aggregator, ShouldSnapshot);
-				Def->OwnedTagsContainer.AppendTags(InMod.Info.PassedTags);
 				NumApplied++;
 				break;
 			}
@@ -133,29 +150,6 @@ int32 FGameplayEffectSpec::ApplyModifiersFrom(const FGameplayEffectSpec &InSpec,
 	return NumApplied;
 }
 
-int32 FGameplayEffectSpec::ApplyModifiersToAggregators(const FModifierSpec &InMod, FAggregatorRef &MyMod, const FModifierQualifier &QualifierContext, bool ShouldSnapshot)
-{
-	int32 NumApplied = 0;
-	for (int32 Op = 0; Op < EGameplayModOp::Max; ++Op)
-	{
-		for (FAggregatorRef SubMod : MyMod.Get()->Mods[Op])
-		{
-			if (InMod.CanModifyAggregator(SubMod, QualifierContext))
-			{
-				SubMod.Get()->ApplyMod(InMod.Info.ModifierOp, InMod.Aggregator, ShouldSnapshot);
-				SubMod.Get()->BaseData.Tags.AppendTags(InMod.Info.PassedTags);
-				NumApplied++;
-			}
-			else
-			{
-				NumApplied += ApplyModifiersToAggregators(InMod, SubMod, QualifierContext, ShouldSnapshot);
-			}
-		}
-	}
-
-	return NumApplied;
-}
-
 int32 FGameplayEffectSpec::ExecuteModifiersFrom(const FGameplayEffectSpec &InSpec, const FModifierQualifier &QualifierContext)
 {
 	int32 NumExecuted = 0;
@@ -233,20 +227,7 @@ bool FModifierSpec::CanModifyModifier(FModifierSpec &Other, const FModifierQuali
 		return false;
 	}
 
-	if (!AreTagRequirementsSatisfied(Other))
-	{
-		return false;
-	}
-
-	return true;
-}
-
-bool FModifierSpec::CanModifyAggregator(FAggregatorRef Ref, const FModifierQualifier &QualifierContext) const
-{
-	if (!AreTagRequirementsSatisfied(Ref))
-	{
-		return false;
-	}
+	// Tag checking is done at the FAggregator level. So all we do here is the attribute check.
 
 	return true;
 }
@@ -264,24 +245,12 @@ void FModifierSpec::ExecuteModOn(FModifierSpec &Other) const
 
 bool FModifierSpec::AreTagRequirementsSatisfied(const FModifierSpec &ModifierToBeModified) const
 {
-	bool HasRequired = ModifierToBeModified.Info.OwnedTags.HasAllTags(Info.RequiredTags);
-	bool HasIgnored = ModifierToBeModified.Info.OwnedTags.Num() > 0 && ModifierToBeModified.Info.OwnedTags.HasAnyTag(Info.IgnoreTags);
-
-	return HasRequired && !HasIgnored;
-}
-
-bool FModifierSpec::AreTagRequirementsSatisfied(const FAggregatorRef &AggregatorToBeModified) const
-{
-	const FAggregator* Aggregator = AggregatorToBeModified.Get();
-	if (Aggregator)
-	{
-		bool HasRequired = Aggregator->BaseData.Tags.HasAllTags(Info.RequiredTags);
-		bool HasIgnored = Aggregator->BaseData.Tags.Num() > 0 && Aggregator->BaseData.Tags.HasAnyTag(Info.IgnoreTags);
+	const FGameplayModifierEvaluatedData & ToBeModifiedData = ModifierToBeModified.Aggregator.Get()->Evaluate();
 
-		return HasRequired && !HasIgnored;
-	}
+	bool HasRequired = ToBeModifiedData.Tags.HasAllTags(this->Info.RequiredTags);
+	bool HasIgnored = ToBeModifiedData.Tags.Num() > 0 && ToBeModifiedData.Tags.HasAnyTag(this->Info.IgnoreTags);
 
-	return false;
+	return HasRequired && !HasIgnored;
 }
 
 // --------------------------------------------------------------------------------------------------------------------------------------------------------
@@ -403,7 +372,7 @@ FAggregator::FAggregator(const FScalableFloat &InBaseMagnitude, TSharedPtr<FGame
 
 FAggregator::FAggregator(const FGameplayModifierEvaluatedData &InEvalData, const TCHAR *InDebugStr)
 : Level(TSharedPtr<FGameplayEffectLevelSpec>(new FGameplayEffectLevelSpec()))
-, BaseData(InEvalData.Magnitude, InEvalData.Callbacks, &InEvalData.Tags)
+, BaseData(InEvalData.Magnitude, InEvalData.Callbacks)
 {
 #if SKILL_SYSTEM_AGGREGATOR_DEBUG
 	if (InDebugStr)
@@ -454,6 +423,8 @@ void FAggregator::RefreshDependencies()
 			}
 		}
 	}
+	
+	RegisterLevelDependancies();
 }
 
 void FAggregator::ApplyMod(EGameplayModOp::Type ModType, FAggregatorRef Ref, bool TakeSnapshot)
@@ -550,6 +521,15 @@ void FAggregator::TakeSnapshotOfLevel()
 	Level = TSharedPtr<FGameplayEffectLevelSpec>( new FGameplayEffectLevelSpec( *Level.Get()) );
 }
 
+void FAggregator::RegisterLevelDependancies()
+{
+	if (!BaseData.Magnitude.IsStatic())
+	{
+		TWeakPtr<FAggregator> LocalWeakPtr(SharedThis(this));
+		Level->RegisterLevelDependancy(LocalWeakPtr);
+	}
+}
+
 // Please try really hard to never add a "force full re-evaluate" flag to this function!
 // We want to strive to make this system dirty the cached data when its actually dirtied, 
 // and never do full catch all rebuilds.
@@ -560,9 +540,9 @@ const FGameplayModifierEvaluatedData& FAggregator::Evaluate() const
 
 	if (!CachedData.IsValid)
 	{
-		CachedData.IsValid = true;
-
+		// ------------------------------------------------------------------------
 		// If there are any overrides, then just take the first valid one.
+		// ------------------------------------------------------------------------
 		for (const FAggregatorRef &Agg : Mods[EGameplayModOp::Override])
 		{
 			SKILL_LOG_SCOPE(TEXT("EGameplayModOp::Override"));
@@ -576,7 +556,7 @@ const FGameplayModifierEvaluatedData& FAggregator::Evaluate() const
 		// If we are going to do math, we need to lock our base value in.
 		// Calculate our magnitude at our level
 		// (We may not have a level, in that case, ensure we also don't have a leveling table)
-		
+
 		float EvaluatedMagnitude = 0.f;
 		if (Level->IsValid())
 		{
@@ -588,41 +568,160 @@ const FGameplayModifierEvaluatedData& FAggregator::Evaluate() const
 		}
 
 		CachedData = FGameplayModifierEvaluatedData(EvaluatedMagnitude, BaseData.Callbacks, ActiveHandle, &BaseData.Tags);
-		SKILL_LOG(Log, TEXT("BaseMagnitude: %.2f"), CachedData.Magnitude);
 
-		// Additive: Just sum them up
-		for (const FAggregatorRef &Agg : Mods[EGameplayModOp::Additive])
+		int32 TotalModCount = Mods[EGameplayModOp::Additive].Num() + Mods[EGameplayModOp::Multiplicitive].Num() + Mods[EGameplayModOp::Division].Num();
+		
+		// Early out if no mods.
+		// We need to calculate num of mods anyways to do ModLIst.Reserve.
+		if (TotalModCount <= 0)
 		{
-			SKILL_LOG_SCOPE(TEXT("EGameplayModOp::Additive"));
-			if (Agg.IsValid())
-			{
-				Agg.Get()->Evaluate().Aggregate( CachedData.Tags, CachedData.Magnitude, 0.f );
-			}
-		}		
+			SKILL_LOG(Log, TEXT("Final Magnitude: %.2f"), CachedData.Magnitude);
+			return CachedData;
+		}
+
+		// ------------------------------------------------------------------------
+		//	Apply Numeric Modifiers
+		//		This is convoluted due to tagging. We have modifiers that require we have or don't have certain tags.
+		//		These mods can also give us new tags. Its possible the 1st mod will require a tag that the 2nd mod gives us.
+		//
+		//	The basic approach here is create an ordered, linear list of all modifiers:
+		//		[Additive][Multipliticitive][Division] mods
+		//
+		//	We make a pass through the list, aggregating as we go. During a pass we keep track of what what tags we've added
+		//	and if there were any mods that we rejected due to not having tags. When the pass is over, we check if we added
+		//	any that we needed. If so, we make another pass. (Once a modifier aggregated, we remove it from the list).
+		//
+		//	Paradoxes are still possible. ModX gives tag A, requires we don't have tag B. ModY gives tag B, requires we don't have tag A.
+		//	We detect this in a single pass. In the above example we would aggregate ModX, but warn loudly when we found ModY.
+		//	(we expect content to solve this via stacking rules, or just not making these type of requirements).
+		//
+		//
+		//
+		//	
+		// ------------------------------------------------------------------------
+		
+		// We have to do tag aggregation to figure out what we can and can't apply.
+		FGameplayTagContainer CommittedIgnoreTags;
 
-		// Multiple - sum the multiplication with a bias then multiply it at the end
-		float Factor = 1.f;
-		for (const FAggregatorRef &Agg : Mods[EGameplayModOp::Multiplicitive])
+		TArray<const FAggregator*> ModList;
+		ModList.Reserve(TotalModCount + (EGameplayModOp::Override - EGameplayModOp::Additive));
+
+		// Build linear list of what we will aggregate
+		for (int32 OpIdx = (EGameplayModOp::Additive); OpIdx < EGameplayModOp::Override; ++OpIdx)
 		{
-			SKILL_LOG_SCOPE(TEXT("EGameplayModOp::Multiplicitive"));
-			if (Agg.IsValid())
-			{
-				Agg.Get()->Evaluate().Aggregate( CachedData.Tags, Factor, 1.f );
+			for (const FAggregatorRef &Agg : Mods[OpIdx])
+			{	
+				ModList.Add(Agg.Get());
 			}
+			ModList.Add(this);	// Sentinel value to signify "NextOp"
 		}
-		CachedData.Magnitude *= Factor;
 
-		// Divisor - sum the division with a bias then divide it at the end
-		Factor = 1.f;
-		for (const FAggregatorRef &Agg : Mods[EGameplayModOp::Division])
+		static const float OpBias[EGameplayModOp::Override] = {
+			0.f,	// EGameplayModOp::Additive
+			1.f,	// EGameplayModOp::Multiplicitive
+			1.f,	// EGameplayModOp::Division
+		};
+
+		float OpAggregation[EGameplayModOp::Override] = {
+			0.f,	// EGameplayModOp::Additive
+			1.f,	// EGameplayModOp::Multiplicitive
+			1.f,	// EGameplayModOp::Division
+		};
+
+		// Make multiple passes to do tagging
+		while (true)
 		{
-			SKILL_LOG_SCOPE(TEXT("EGameplayModOp::Division"));
-			if (Agg.IsValid())
+			FGameplayTagContainer	NewGiveTags;
+			FGameplayTagContainer	NewIgnoreTags;
+			FGameplayTagContainer	MissingTags;
+
+			EGameplayModOp::Type	ModOp = (EGameplayModOp::Additive);
+			
+			for (int32 ModIdx = 0; ModIdx < ModList.Num()-1; ++ModIdx)
 			{
-				Agg.Get()->Evaluate().Aggregate(CachedData.Tags, Factor, 1.f);
+				const FAggregator * Agg = ModList[ModIdx];
+				if (Agg == NULL)
+				{
+					continue;
+				}
+				if (Agg == this)
+				{
+					ModOp = static_cast<EGameplayModOp::Type>(static_cast<int32>(ModOp)+1);
+					check(ModOp < EGameplayModOp::Override);
+					continue;
+				}
+				
+				const FGameplayModifierData &ModBaseData = ModList[ModIdx]->BaseData;
+				const FGameplayTagContainer &ModIgnoreTags = ModBaseData.IgnoreTags;
+				const FGameplayTagContainer &ModRequireTags = ModBaseData.RequireTags;
+
+				// This mod requires we have certain tags
+				if (ModRequireTags.Num() > 0 && !CachedData.Tags.HasAllTags(ModRequireTags) && !NewGiveTags.HasAllTags(ModRequireTags))
+				{
+					// But something else could give us this tag! So keep track of it and don't remove this Mod from the ModList.
+					MissingTags.AppendTags(ModRequireTags);
+					continue;
+				}
+
+				// This mod is now either accepted or rejected. It will not be needed for subsequent passes in this Evaluate, so NULL It out now.
+				ModList[ModIdx] = NULL;
+				
+				// This mod requires we don't have certain tags
+				if (ModIgnoreTags.Num() > 0 && CachedData.Tags.HasAnyTag(ModIgnoreTags))
+				{
+					continue;
+				}
+
+				// Check for conflicts within this pass
+				if (ModIgnoreTags.Num() > 0 && NewGiveTags.Num() > 0 && ModIgnoreTags.HasAnyTag(NewGiveTags))
+				{
+					// Pass problem!
+					SKILL_LOG(Warning, TEXT("Tagging conflicts during Aggregate! Use Stacking rules to avoid this"));
+					SKILL_LOG(Warning, TEXT("   While Evaluating: %s"), *ToSimpleString());
+					SKILL_LOG(Warning, TEXT("   Applying Mod: %s"), *Agg->ToSimpleString());
+					SKILL_LOG(Warning, TEXT("   ModIgnoreTags: %s"), *ModIgnoreTags.ToString() );
+					SKILL_LOG(Warning, TEXT("   NewGiveTags: %s"), *NewGiveTags.ToString());
+					continue;
+				}
+
+				// Our requirements on the mod
+				const FGameplayModifierEvaluatedData& ModEvaluatedData = Agg->Evaluate();
+
+				// We have already commited during this aggregation to not have certain tags
+				if (CommittedIgnoreTags.Num() > 0 && CommittedIgnoreTags.HasAnyTag(ModEvaluatedData.Tags))
+				{
+					continue;
+				}
+
+				if (NewIgnoreTags.Num() > 0 && ModEvaluatedData.Tags.Num() > 0 && NewIgnoreTags.HasAnyTag(ModEvaluatedData.Tags))
+				{
+					// Pass problem!
+					SKILL_LOG(Warning, TEXT("Tagging conflicts during Aggregate! Use Stacking rules to avoid this"));
+					SKILL_LOG(Warning, TEXT("   While Evaluating: %s"), *ToSimpleString());
+					SKILL_LOG(Warning, TEXT("   Applying Mod: %s"), *Agg->ToSimpleString());
+					SKILL_LOG(Warning, TEXT("   Mod Tags: %s"), *ModEvaluatedData.Tags.ToString());
+					SKILL_LOG(Warning, TEXT("   NewIgnoreTags: %s"), *NewIgnoreTags.ToString());
+					continue;
+				}
+
+				// Commit this Mod
+				NewIgnoreTags.AppendTags(ModIgnoreTags);
+				ModEvaluatedData.Aggregate(NewGiveTags, OpAggregation[ModOp], OpBias[ModOp]);
+			}
+
+			// Commit this pass's tags
+			CachedData.Tags.AppendTags(NewGiveTags);
+
+			// Keep doing passes until we don't add any new tags or don't have any missing tags
+			if (MissingTags.Num() <= 0 || !MissingTags.HasAnyTag( NewGiveTags ))
+			{
+				break;
 			}
 		}
-		CachedData.Magnitude /= Factor;
+
+		float Division = OpAggregation[EGameplayModOp::Division] > 0.f ? OpAggregation[EGameplayModOp::Division] : 1.f;
+
+		CachedData.Magnitude = ((CachedData.Magnitude + OpAggregation[EGameplayModOp::Additive]) * OpAggregation[EGameplayModOp::Multiplicitive]) / Division;
 
 		SKILL_LOG(Log, TEXT("Final Magnitude: %.2f"), CachedData.Magnitude);
 		
@@ -661,6 +760,11 @@ void FAggregator::PostEvaluate(const struct FGameplayEffectModCallbackData &Data
 FAggregator & FAggregator::MarkDirty()
 {
 	CachedData.IsValid = false;
+
+	// Execute OnDirty callbacks first. This may do things like update the actual uproperty value of an attribute.
+	OnDirty.ExecuteIfBound(this);
+
+	// Now tell people who depend on my value that I have changed. Important to do this after the OnDirty callback has been called.
 	for (int32 i=0; i < Dependants.Num(); ++i)
 	{
 		TWeakPtr<FAggregator> WeakPtr = Dependants[i];
@@ -678,7 +782,6 @@ FAggregator & FAggregator::MarkDirty()
 			--i;
 		}
 	}
-	OnDirty.ExecuteIfBound(this);
 
 	return *this;
 }
@@ -697,6 +800,12 @@ void FAggregator::MakeUniqueDeep()
 	}
 }
 
+void FAggregator::CleaerAllDependancies()
+{
+	Dependants.SetNum(0);
+	OnDirty.Unbind();
+}
+
 // --------------------------------------------------------------------------------------------------------------------------------------------------------
 //
 //	FGameplayModifierEvaluatedData
@@ -802,24 +911,11 @@ void FActiveGameplayEffectsContainer::ApplySpecToActiveEffectsAndAttributes(cons
 
 			SKILL_LOG_SCOPE(TEXT("Applying Attribute Mod %s to property"), *Mod.ToSimpleString());
 
-			FAggregatorRef *RefPtr = OngoingPropertyEffects.Find(Mod.Info.Attribute);
-			if (!RefPtr)
-			{
-				// Create a new aggregator for this attribute.
-
-				float CurrentValueOfProperty = Owner->GetNumericAttribute(Mod.Info.Attribute);
-				SKILL_LOG(Log, TEXT("Creating new entry in OngoingPropertyEffect map. CurrentValue: %.2f"), CurrentValueOfProperty);
-
-				FAggregator *NewPropertyAggregator = new FAggregator(FGameplayModifierEvaluatedData(CurrentValueOfProperty), SKILL_AGG_DEBUG(TEXT("Attribute %s Aggregator"), *Mod.Info.Attribute.GetName()));
-
-				NewPropertyAggregator->OnDirty = FAggregator::FOnDirty::CreateRaw(this, &FActiveGameplayEffectsContainer::OnPropertyAggregatorDirty, Mod.Info.Attribute);
-
-				RefPtr = &OngoingPropertyEffects.Add(Mod.Info.Attribute, FAggregatorRef(NewPropertyAggregator));
-			}
+			FAggregator * Aggregator = FindOrCreateAttributeAggregator(Mod.Info.Attribute).Get();
 
 			// Add the modifier to the property value aggregator
 			// Note that this will immediately invoke the callback to update the attribute's current value, so we don't have to explicitly do it here.
-			RefPtr->Get()->ApplyMod(Mod.Info.ModifierOp, Mod.Aggregator, Spec.ShouldApplyAsSnapshot(QualifierContext));
+			Aggregator->ApplyMod(Mod.Info.ModifierOp, Mod.Aggregator, Spec.ShouldApplyAsSnapshot(QualifierContext));
 		}
 
 		if (Mod.Info.ModifierType == EGameplayMod::ActiveGE)
@@ -952,6 +1048,31 @@ bool FActiveGameplayEffectsContainer::ExecuteGameplayEffect(FActiveGameplayEffec
 	return false;
 }
 
+void FActiveGameplayEffectsContainer::AddDependancyToAttribute(FGameplayAttribute Attribute, const TWeakPtr<FAggregator> InDependant)
+{
+	FAggregator * Aggregator = FindOrCreateAttributeAggregator(Attribute).Get();
+	Aggregator->AddDependantAggregator(InDependant);
+}
+
+FAggregatorRef & FActiveGameplayEffectsContainer::FindOrCreateAttributeAggregator(FGameplayAttribute Attribute)
+{
+	FAggregatorRef *RefPtr = OngoingPropertyEffects.Find(Attribute);
+	if (RefPtr)
+	{
+		return *RefPtr;
+	}
+
+	// Create a new aggregator for this attribute.
+	float CurrentValueOfProperty = Owner->GetNumericAttribute(Attribute);
+	SKILL_LOG(Log, TEXT("Creating new entry in OngoingPropertyEffect map for AddDependancyToAttribute. CurrentValue: %.2f"), CurrentValueOfProperty);
+
+	FAggregator *NewPropertyAggregator = new FAggregator(FGameplayModifierEvaluatedData(CurrentValueOfProperty), SKILL_AGG_DEBUG(TEXT("Attribute %s Aggregator"), *Attribute.GetName()));
+
+	NewPropertyAggregator->OnDirty = FAggregator::FOnDirty::CreateRaw(this, &FActiveGameplayEffectsContainer::OnPropertyAggregatorDirty, Attribute);
+
+	return OngoingPropertyEffects.Add(Attribute, FAggregatorRef(NewPropertyAggregator));
+}
+
 float FActiveGameplayEffectsContainer::GetGameplayEffectDuration(FActiveGameplayEffectHandle Handle) const
 {
 	// Could make this a map for quicker lookup
@@ -1064,17 +1185,24 @@ bool FActiveGameplayEffectsContainer::RemoveActiveGameplayEffect(FActiveGameplay
 	return false;
 }
 
-void FActiveGameplayEffectsContainer::RegisterAttributeModifyCallback(FGameplayAttribute Attribute, FOnGameplayAttributeEffectExecuted Delegate)
-{
-	PropertyExecutionCallbacks.FindOrAdd(Attribute) = Delegate;
-}
-
 bool FActiveGameplayEffectsContainer::IsNetAuthority() const
 {
 	check(Owner);
 	return Owner->IsOwnerActorAuthoritative();
 }
 
+void FActiveGameplayEffectsContainer::PreDestroy()
+{
+	// Prior to destruction we need to clear all dependancies
+	for (auto It : OngoingPropertyEffects)
+	{
+		if (It.Value.IsValid())
+		{
+			It.Value.Get()->CleaerAllDependancies();
+		}
+	}
+}
+
 void FActiveGameplayEffectsContainer::TEMP_TickActiveEffects(float DeltaSeconds)
 {
 	if (!IsNetAuthority())
@@ -1123,10 +1251,52 @@ void FActiveGameplayEffectsContainer::TEMP_TickActiveEffects(float DeltaSeconds)
 
 // --------------------------------------------------------------------------------------------------------------------------------------------------------
 //
-//	FGameplayEffectSystemManager
+//	FGameplayEffectLevelSpec
 //
 // --------------------------------------------------------------------------------------------------------------------------------------------------------
 
+float FGameplayEffectLevelSpec::GetLevel() const
+{
+	if (ConstantLevel != INVALID_LEVEL)
+	{
+		return ConstantLevel;
+	}
+
+	if (Source.IsValid() && Attribute.GetUProperty() != NULL)
+	{
+		UClass * SetClass = Attribute.GetAttributeSetClass();
+		check(SetClass);
+
+		TArray<UObject*> ChildObjects;
+		GetObjectsWithOuter(Source.Get(), ChildObjects);
+
+		for (int32 ChildIndex = 0; ChildIndex < ChildObjects.Num(); ++ChildIndex)
+		{
+			UObject *Obj = ChildObjects[ChildIndex];
+			if (Obj->GetClass()->IsChildOf(SetClass))
+			{
+				CachedLevel = Attribute.GetNumericValueChecked(CastChecked<UAttributeSet>(Obj));
+				return CachedLevel;
+			}
+		}
+	}
+
+	// Our source is invalid. 
+	ConstantLevel = CachedLevel;
+	return CachedLevel;
+}
+
+void FGameplayEffectLevelSpec::RegisterLevelDependancy(TWeakPtr<FAggregator> OwningAggregator)
+{
+	if (Source.IsValid() && Attribute.GetUProperty() != NULL)
+	{
+		UAttributeComponent * SourceComponent = Source->FindComponentByClass<UAttributeComponent>();
+		if (SourceComponent)
+		{
+			SourceComponent->AddDependancyToAttribute(Attribute, OwningAggregator);
+		}
+	}
+}
 
 // --------------------------------------------------------------------------------------------------------------------------------------------------------
 //
@@ -1157,3 +1327,4 @@ FString EGameplayEffectCopyPolicyToString(int32 Type)
 	static UEnum *e = FindObject<UEnum>(ANY_PACKAGE, TEXT("EGameplayEffectCopyPolicy"));
 	return e->GetEnum(Type).ToString();
 }
+
diff --git a/Engine/Source/Runtime/SkillSystem/Private/GameplayEffectTests.cpp b/Engine/Source/Runtime/SkillSystem/Private/GameplayEffectTests.cpp
--- a/Engine/Source/Runtime/SkillSystem/Private/GameplayEffectTests.cpp
+++ b/Engine/Source/Runtime/SkillSystem/Private/GameplayEffectTests.cpp
@@ -120,7 +120,7 @@ bool GameplayEffectsTest_InstantDamage(UWorld *World, FAutomationTestBase * Test
 
 		SourceComponent->ApplyGameplayEffectSpecToTarget(BaseDmgEffect, DestComponent, 1.f);
 		
-		Test->TestTrue(TEXT("Basic Instant Damage Applied"), (DestComponent->GetSet<USkillSystemTestAttributeSet>()->Health == (StartHealth + DamageValue)));
+		Test->TestTrue(SKILL_TEST_TEXT("Basic Instant Damage Applied"), (DestComponent->GetSet<USkillSystemTestAttributeSet>()->Health == (StartHealth + DamageValue)));
 		SKILL_LOG(Log, TEXT("Final Health: %.2f"), DestComponent->GetSet<USkillSystemTestAttributeSet>()->Health);
 	}
 
@@ -228,7 +228,7 @@ bool GameplayEffectsTest_InstantDamage_Buffed(UWorld *World, FAutomationTestBase
 
 		float ExpectedValue = (StartHealth + (DamageValue * BonusDamageMultiplier));
 
-		Test->TestTrue(TEXT("Buff Instant Damage Applied"), (DestComponent->GetSet<USkillSystemTestAttributeSet>()->Health == ExpectedValue));
+		Test->TestTrue(SKILL_TEST_TEXT("Buff Instant Damage Applied"), (DestComponent->GetSet<USkillSystemTestAttributeSet>()->Health == ExpectedValue));
 		SKILL_LOG(Log, TEXT("Final Health: %.2f"), DestComponent->GetSet<USkillSystemTestAttributeSet>()->Health);
 	}
 
@@ -274,7 +274,7 @@ bool GameplayEffectsTest_TemporaryDamage(UWorld *World, FAutomationTestBase * Te
 
 		float ExpectedValue = (StartHealth + DamageValue);
 
-		Test->TestTrue(TEXT("INFINITE_DURATION Damage Applied"), (DestComponent->GetSet<USkillSystemTestAttributeSet>()->Health == ExpectedValue));
+		Test->TestTrue(SKILL_TEST_TEXT("INFINITE_DURATION Damage Applied"), (DestComponent->GetSet<USkillSystemTestAttributeSet>()->Health == ExpectedValue));
 		SKILL_LOG(Log, TEXT("After Damage Applied: Health: %.2f"), DestComponent->GetSet<USkillSystemTestAttributeSet>()->Health);
 	}
 
@@ -283,7 +283,7 @@ bool GameplayEffectsTest_TemporaryDamage(UWorld *World, FAutomationTestBase * Te
 		bool RemovedEffect = DestComponent->RemoveActiveGameplayEffect(AppliedHandle);
 		float ExpectedValue = StartHealth;
 
-		Test->TestTrue(TEXT("INFINITE_DURATION Damage Applied"), (DestComponent->GetSet<USkillSystemTestAttributeSet>()->Health == ExpectedValue));
+		Test->TestTrue(SKILL_TEST_TEXT("INFINITE_DURATION Damage Applied"), (DestComponent->GetSet<USkillSystemTestAttributeSet>()->Health == ExpectedValue));
 		SKILL_LOG(Log, TEXT("After Removal. Health: %.2f. RemovedEffecte: %d"), DestComponent->GetSet<USkillSystemTestAttributeSet>()->Health, RemovedEffect);
 	}
 
@@ -330,10 +330,10 @@ bool GameplayEffectsTest_TemporaryDamageBuffed(UWorld *World, FAutomationTestBas
 
 		float ExpectedValue = (StartHealth + DamageValue);
 
-		Test->TestTrue(TEXT("INFINITE_DURATION Damage Applied"), (DestComponent->GetSet<USkillSystemTestAttributeSet>()->Health == ExpectedValue));
+		Test->TestTrue(SKILL_TEST_TEXT("INFINITE_DURATION Damage Applied"), (DestComponent->GetSet<USkillSystemTestAttributeSet>()->Health == ExpectedValue));
 		SKILL_LOG(Log, TEXT("After Damage Applied: Health: %.2f"), DestComponent->GetSet<USkillSystemTestAttributeSet>()->Health);
 
-		Test->TestTrue(TEXT("Number of GameplayEffects=1"), DestComponent->GetNumActiveGameplayEffect() == 1);
+		Test->TestTrue(SKILL_TEST_TEXT("Number of GameplayEffects=1"), DestComponent->GetNumActiveGameplayEffect() == 1);
 	}
 
 	// Now Buff the GameplayEffect we just added and confirm the health removal is increased 2x
@@ -354,11 +354,11 @@ bool GameplayEffectsTest_TemporaryDamageBuffed(UWorld *World, FAutomationTestBas
 
 		float ExpectedValue = (StartHealth + (DamageValue * DamageBuffMultiplier));
 
-		Test->TestTrue(TEXT("INFINITE_DURATION Buffed Damage Applied"), (DestComponent->GetSet<USkillSystemTestAttributeSet>()->Health == ExpectedValue));
+		Test->TestTrue(SKILL_TEST_TEXT("INFINITE_DURATION Buffed Damage Applied"), (DestComponent->GetSet<USkillSystemTestAttributeSet>()->Health == ExpectedValue));
 		SKILL_LOG(Log, TEXT("After Damage Applied: Health: %.2f"), DestComponent->GetSet<USkillSystemTestAttributeSet>()->Health);
 
 		// Confirm still only 1 active GE (since this was instant application
-		Test->TestTrue(TEXT("Number of GameplayEffects=1"), DestComponent->GetNumActiveGameplayEffect() == 1);
+		Test->TestTrue(SKILL_TEST_TEXT("Number of GameplayEffects=1"), DestComponent->GetNumActiveGameplayEffect() == 1);
 	}
 
 
@@ -367,11 +367,11 @@ bool GameplayEffectsTest_TemporaryDamageBuffed(UWorld *World, FAutomationTestBas
 		bool RemovedEffect = DestComponent->RemoveActiveGameplayEffect(AppliedHandle);
 		float ExpectedValue = StartHealth;
 
-		Test->TestTrue(TEXT("INFINITE_DURATION Damage Applied"), (DestComponent->GetSet<USkillSystemTestAttributeSet>()->Health == ExpectedValue));
+		Test->TestTrue(SKILL_TEST_TEXT("INFINITE_DURATION Damage Applied"), (DestComponent->GetSet<USkillSystemTestAttributeSet>()->Health == ExpectedValue));
 		SKILL_LOG(Log, TEXT("After Removal. Health: %.2f. RemovedEffecte: %d"), DestComponent->GetSet<USkillSystemTestAttributeSet>()->Health, RemovedEffect);
 
 		// Confirm no more GEs
-		Test->TestTrue(TEXT("Number of GameplayEffects=0"), DestComponent->GetNumActiveGameplayEffect() == 0);
+		Test->TestTrue(SKILL_TEST_TEXT("Number of GameplayEffects=0"), DestComponent->GetNumActiveGameplayEffect() == 0);
 	}
 
 	return true;
@@ -419,10 +419,10 @@ bool GameplayEffectsTest_TemporaryDamageTemporaryBuff(UWorld *World, FAutomation
 
 		float ExpectedValue = (StartHealth + DamageValue);
 
-		Test->TestTrue(TEXT("GameplayEffectsTest_TemporaryDamageTemporaryBuff INFINITE_DURATION Damage Applied"), (DestComponent->GetSet<USkillSystemTestAttributeSet>()->Health == ExpectedValue));
+		Test->TestTrue(SKILL_TEST_TEXT("GameplayEffectsTest_TemporaryDamageTemporaryBuff INFINITE_DURATION Damage Applied"), (DestComponent->GetSet<USkillSystemTestAttributeSet>()->Health == ExpectedValue));
 		SKILL_LOG(Log, TEXT("After Damage Applied: Health: %.2f"), DestComponent->GetSet<USkillSystemTestAttributeSet>()->Health);
 
-		Test->TestTrue(TEXT("GameplayEffectsTest_TemporaryDamageTemporaryBuff Number of GameplayEffects=1"), DestComponent->GetNumActiveGameplayEffect() == 1);
+		Test->TestTrue(SKILL_TEST_TEXT("GameplayEffectsTest_TemporaryDamageTemporaryBuff Number of GameplayEffects=1"), DestComponent->GetNumActiveGameplayEffect() == 1);
 	}
 
 	// Now Buff the GameplayEffect we just added and confirm the health removal is increased 2x
@@ -444,11 +444,11 @@ bool GameplayEffectsTest_TemporaryDamageTemporaryBuff(UWorld *World, FAutomation
 
 		float ExpectedValue = (StartHealth + (DamageValue * DamageBuffMultiplier));
 
-		Test->TestTrue(TEXT("GameplayEffectsTest_TemporaryDamageTemporaryBuff INFINITE_DURATION Buffed Damage Applied"), (DestComponent->GetSet<USkillSystemTestAttributeSet>()->Health == ExpectedValue));
+		Test->TestTrue(SKILL_TEST_TEXT("GameplayEffectsTest_TemporaryDamageTemporaryBuff INFINITE_DURATION Buffed Damage Applied"), (DestComponent->GetSet<USkillSystemTestAttributeSet>()->Health == ExpectedValue));
 		SKILL_LOG(Log, TEXT("After Damage Applied: Health: %.2f"), DestComponent->GetSet<USkillSystemTestAttributeSet>()->Health);
 
 		// Confirm there are 2 GEs
-		Test->TestTrue(TEXT("GameplayEffectsTest_TemporaryDamageTemporaryBuff Number of GameplayEffects=1"), DestComponent->GetNumActiveGameplayEffect() == 2);
+		Test->TestTrue(SKILL_TEST_TEXT("GameplayEffectsTest_TemporaryDamageTemporaryBuff Number of GameplayEffects=1"), DestComponent->GetNumActiveGameplayEffect() == 2);
 	}
 
 	// Print out the whole enchillada
@@ -461,11 +461,11 @@ bool GameplayEffectsTest_TemporaryDamageTemporaryBuff(UWorld *World, FAutomation
 		bool RemovedEffect = DestComponent->RemoveActiveGameplayEffect(AppliedHandleBuff);
 		float ExpectedValue = StartHealth + DamageValue;
 
-		Test->TestTrue(TEXT("GameplayEffectsTest_TemporaryDamageTemporaryBuff INFINITE_DURATION Damage Buff Removed"), (DestComponent->GetSet<USkillSystemTestAttributeSet>()->Health == ExpectedValue));
+		Test->TestTrue(SKILL_TEST_TEXT("GameplayEffectsTest_TemporaryDamageTemporaryBuff INFINITE_DURATION Damage Buff Removed"), (DestComponent->GetSet<USkillSystemTestAttributeSet>()->Health == ExpectedValue));
 		SKILL_LOG(Log, TEXT("After Buff Removal. Health: %.2f. RemovedEffecte: %d"), DestComponent->GetSet<USkillSystemTestAttributeSet>()->Health, RemovedEffect);
 
 		// Confirm 1 more GEs
-		Test->TestTrue(TEXT("Number of GameplayEffects=1"), DestComponent->GetNumActiveGameplayEffect() == 1);
+		Test->TestTrue(SKILL_TEST_TEXT("Number of GameplayEffects=1"), DestComponent->GetNumActiveGameplayEffect() == 1);
 	}
 
 	// Remove the damage GE
@@ -473,11 +473,11 @@ bool GameplayEffectsTest_TemporaryDamageTemporaryBuff(UWorld *World, FAutomation
 		bool RemovedEffect = DestComponent->RemoveActiveGameplayEffect(AppliedHandle);
 		float ExpectedValue = StartHealth;
 
-		Test->TestTrue(TEXT("GameplayEffectsTest_TemporaryDamageTemporaryBuff INFINITE_DURATION Damage Removed"), (DestComponent->GetSet<USkillSystemTestAttributeSet>()->Health == ExpectedValue));
+		Test->TestTrue(SKILL_TEST_TEXT("GameplayEffectsTest_TemporaryDamageTemporaryBuff INFINITE_DURATION Damage Removed"), (DestComponent->GetSet<USkillSystemTestAttributeSet>()->Health == ExpectedValue));
 		SKILL_LOG(Log, TEXT("After Removal. Health: %.2f. RemovedEffecte: %d"), DestComponent->GetSet<USkillSystemTestAttributeSet>()->Health, RemovedEffect);
 
 		// Confirm no more GEs
-		Test->TestTrue(TEXT("GameplayEffectsTest_TemporaryDamageTemporaryBuff Number of GameplayEffects=0"), DestComponent->GetNumActiveGameplayEffect() == 0);
+		Test->TestTrue(SKILL_TEST_TEXT("GameplayEffectsTest_TemporaryDamageTemporaryBuff Number of GameplayEffects=0"), DestComponent->GetNumActiveGameplayEffect() == 0);
 	}
 
 	return true;
@@ -1787,7 +1787,7 @@ bool GameplayEffectsTest_ShieldlExtensionMultiple(UWorld *World, FAutomationTest
 
 UCurveTable * SetGlobalCurveTable()
 {
-	FString CSV(TEXT(", 1, 100\r\nStandardHealth, 1, 100\r\nStandardDamage, 1, 100"));
+	FString CSV(TEXT(", 0, 1, 100\r\nStandardHealth, 0, 1, 100\r\nStandardDamage, 0, 1, 100\r\nLinearCurve, 0, 1, 100"));
 
 	UCurveTable * CurveTable = Cast<UCurveTable>(StaticConstructObject(UCurveTable::StaticClass(), GetTransientPackage(), FName(TEXT("TempCurveTable"))));
 	CurveTable->CreateTableFromCSVString(CSV);
@@ -1810,7 +1810,7 @@ void ClearGlobalCurveTable()
 
 UCurveTable * GetStandardDamageOverrideCurveTable(float Factor)
 {
-	FString CSV = FString::Printf(TEXT(", 1, 100\r\nStandardDamage, %.2f, %.2f"), Factor * 1.f, Factor * 100.f);
+	FString CSV = FString::Printf(TEXT(", 0, 1, 100\r\nStandardDamage, 0, %.2f, %.2f"), Factor * 1.f, Factor * 100.f);
 
 	UCurveTable * CurveTable = Cast<UCurveTable>(StaticConstructObject(UCurveTable::StaticClass(), GetTransientPackage(), FName(TEXT("TempCurveTable"))));
 	CurveTable->CreateTableFromCSVString(CSV);
@@ -2013,14 +2013,14 @@ bool GameplayEffectsTest_InstantDamageRequiredTag(UWorld *World, FAutomationTest
 		BaseDmgEffect->Modifiers[0].Attribute.SetUProperty(HealthProperty);
 		BaseDmgEffect->Modifiers[0].OwnedTags.AddTag(FName(TEXT("Damage.Basic")));
 		BaseDmgEffect->Duration.Value = UGameplayEffect::INSTANT_APPLICATION;
-		BaseDmgEffect->OwnedTagsContainer.AddTag(FName(TEXT("Damage.Type1")));
+		BaseDmgEffect->GameplayEffectTags.AddTag(FName(TEXT("Damage.Type1")));
 
 		// Apply to target
 		SourceComponent->ApplyGameplayEffectSpecToTarget(BaseDmgEffect, DestComponent, 1.f);
 
 		float ExpectedValue = (StartHealth + DamageValue);
-
-		Test->TestTrue(TEXT("Instant Damage Required Tag No Protection"), (DestComponent->GetSet<USkillSystemTestAttributeSet>()->Health == ExpectedValue));
+		
+		Test->TestTrue(SKILL_TEST_TEXT("Instant Damage Required Tag No Protection"), (DestComponent->GetSet<USkillSystemTestAttributeSet>()->Health == ExpectedValue));
 		SKILL_LOG(Log, TEXT("Final Health: %.2f"), DestComponent->GetSet<USkillSystemTestAttributeSet>()->Health);
 	}
 
@@ -2039,15 +2039,15 @@ bool GameplayEffectsTest_InstantDamageRequiredTag(UWorld *World, FAutomationTest
 		BaseDmgEffect->Modifiers[0].Attribute.SetUProperty(HealthProperty);
 		BaseDmgEffect->Modifiers[0].OwnedTags.AddTag(FName(TEXT("Damage.Basic")));
 		BaseDmgEffect->Duration.Value = UGameplayEffect::INSTANT_APPLICATION;
-		BaseDmgEffect->OwnedTagsContainer.AddTag(FName(TEXT("Damage.Type2")));
+		BaseDmgEffect->GameplayEffectTags.AddTag(FName(TEXT("Damage.Type2")));
 
 		// Apply to target
 		SourceComponent->ApplyGameplayEffectSpecToTarget(BaseDmgEffect, DestComponent, 1.f);
 
 		float ExpectedValue = (StartHealth + (DamageValue / DamageProtectionDivisor));
-
-		Test->TestTrue(TEXT("Instant Damage Required Tag Protected"), (DestComponent->GetSet<USkillSystemTestAttributeSet>()->Health == ExpectedValue));
-		SKILL_LOG(Log, TEXT("Final Health: %.2f"), DestComponent->GetSet<USkillSystemTestAttributeSet>()->Health);
+		float ActualValue = DestComponent->GetSet<USkillSystemTestAttributeSet>()->Health;
+		
+		Test->TestTrue(SKILL_TEST_TEXT("Damage Applied. Health: Actual: %.2f == Exected: %.2f", ActualValue, ExpectedValue), ActualValue == ExpectedValue);
 	}
 
 	return true;
@@ -2100,14 +2100,14 @@ bool GameplayEffectsTest_InstantDamageIgnoreTag(UWorld *World, FAutomationTestBa
 		BaseDmgEffect->Modifiers[0].Attribute.SetUProperty(HealthProperty);
 		BaseDmgEffect->Modifiers[0].OwnedTags.AddTag(FName(TEXT("Damage.Basic")));
 		BaseDmgEffect->Duration.Value = UGameplayEffect::INSTANT_APPLICATION;
-		BaseDmgEffect->OwnedTagsContainer.AddTag(FName(TEXT("Damage.Type1")));
+		BaseDmgEffect->GameplayEffectTags.AddTag(FName(TEXT("Damage.Type1")));
 
 		// Apply to target
 		SourceComponent->ApplyGameplayEffectSpecToTarget(BaseDmgEffect, DestComponent, 1.f);
 
 		float ExpectedValue = (StartHealth + DamageValue);
 
-		Test->TestTrue(TEXT("Instant Damage Ignore Tag No Protection"), (DestComponent->GetSet<USkillSystemTestAttributeSet>()->Health == ExpectedValue));
+		Test->TestTrue(SKILL_TEST_TEXT("Instant Damage Ignore Tag No Protection"), (DestComponent->GetSet<USkillSystemTestAttributeSet>()->Health == ExpectedValue));
 		SKILL_LOG(Log, TEXT("Final Health: %.2f"), DestComponent->GetSet<USkillSystemTestAttributeSet>()->Health);
 	}
 
@@ -2126,15 +2126,14 @@ bool GameplayEffectsTest_InstantDamageIgnoreTag(UWorld *World, FAutomationTestBa
 		BaseDmgEffect->Modifiers[0].Attribute.SetUProperty(HealthProperty);
 		BaseDmgEffect->Modifiers[0].OwnedTags.AddTag(FName(TEXT("Damage.Basic")));
 		BaseDmgEffect->Duration.Value = UGameplayEffect::INSTANT_APPLICATION;
-		BaseDmgEffect->OwnedTagsContainer.AddTag(FName(TEXT("Damage.Type2")));
-		//		BaseDmgEffect->GameplayEffectTags.AddTag(FName(TEXT("Damage.Type2")));
+		BaseDmgEffect->GameplayEffectTags.AddTag(FName(TEXT("Damage.Type2")));
 
 		// Apply to target
 		SourceComponent->ApplyGameplayEffectSpecToTarget(BaseDmgEffect, DestComponent, 1.f);
 
 		float ExpectedValue = (StartHealth + (DamageValue / DamageProtectionDivisor));
 
-		Test->TestTrue(TEXT("Instant Damage Ignore Tag Protected"), (DestComponent->GetSet<USkillSystemTestAttributeSet>()->Health == ExpectedValue));
+		Test->TestTrue(SKILL_TEST_TEXT("Instant Damage Ignore Tag Protected"), (DestComponent->GetSet<USkillSystemTestAttributeSet>()->Health == ExpectedValue));
 		SKILL_LOG(Log, TEXT("Final Health: %.2f"), DestComponent->GetSet<USkillSystemTestAttributeSet>()->Health);
 	}
 
@@ -2168,8 +2167,7 @@ bool GameplayEffectsTest_InstantDamageModifierPassesTag(UWorld *World, FAutomati
 		BaseDmgEffect->Modifiers[0].ModifierType = EGameplayMod::OutgoingGE;
 		BaseDmgEffect->Modifiers[0].ModifierOp = EGameplayModOp::Multiplicitive;
 		BaseDmgEffect->Modifiers[0].Attribute.SetUProperty(HealthProperty);
-		BaseDmgEffect->Modifiers[0].OwnedTags.AddTag(FName(TEXT("Damage.Basic")));
-		BaseDmgEffect->Modifiers[0].PassedTags.AddTag(FName(TEXT("Damage.Type1")));
+		BaseDmgEffect->Modifiers[0].OwnedTags.AddTag(FName(TEXT("Damage.Type1")));		// When I am applied, the damage modifier gets this tag.
 		BaseDmgEffect->Duration.SetValue(UGameplayEffect::INFINITE_DURATION);
 
 		// Apply to self
@@ -2187,9 +2185,9 @@ bool GameplayEffectsTest_InstantDamageModifierPassesTag(UWorld *World, FAutomati
 		BaseProtectEffect->Modifiers[0].ModifierOp = EGameplayModOp::Division;
 		BaseProtectEffect->Modifiers[0].Attribute.SetUProperty(HealthProperty);
 		BaseProtectEffect->Modifiers[0].OwnedTags.AddTag(FName(TEXT("Protect.Damage")));
+		BaseProtectEffect->Modifiers[0].RequiredTags.AddTag(FName(TEXT("Damage.Type1")));
 		BaseProtectEffect->Duration.SetValue(UGameplayEffect::INFINITE_DURATION);
 		BaseProtectEffect->CopyPolicy = EGameplayEffectCopyPolicy::AlwaysLink;
-		BaseProtectEffect->GameplayEffectRequiredTags.AddTag(FName(TEXT("Damage.Type1")));
 
 		// Apply to self
 		DestComponent->ApplyGameplayEffectSpecToTarget(BaseProtectEffect, DestComponent, 1.f);
@@ -2212,8 +2210,9 @@ bool GameplayEffectsTest_InstantDamageModifierPassesTag(UWorld *World, FAutomati
 		SourceComponent->ApplyGameplayEffectSpecToTarget(BaseDmgEffect, DestComponent, 1.f);
 
 		float ExpectedValue = (StartHealth + ((DamageValue * BonusDamageMultiplier) / DamageProtectionDivisor));
+		float ActualValue = DestComponent->GetSet<USkillSystemTestAttributeSet>()->Health;
 
-		Test->TestTrue(TEXT("Instant Damage Required Tag No Protection"), (DestComponent->GetSet<USkillSystemTestAttributeSet>()->Health == ExpectedValue));
+		Test->TestTrue(SKILL_TEST_TEXT("Instant Damage Required Tag No Protection.  Actual: %.2f == Exected: %.2f", ActualValue, ExpectedValue), (ActualValue == ExpectedValue));
 		SKILL_LOG(Log, TEXT("Final Health: %.2f"), DestComponent->GetSet<USkillSystemTestAttributeSet>()->Health);
 	}
 
@@ -2289,10 +2288,443 @@ bool GameplayEffectsTest_InstantDamageModifierTag(UWorld *World, FAutomationTest
 		// Apply to target
 		SourceComponent->ApplyGameplayEffectSpecToTarget(BaseDmgEffect, DestComponent, 1.f);
 
-		float ExpectedValue = (StartHealth + (DamageValue + (BonusDamageValue / DamageProtectionDivisor)));
+		float ExpectedValue = (StartHealth + ((DamageValue + BonusDamageValue) / DamageProtectionDivisor));
+		float ActualValue = DestComponent->GetSet<USkillSystemTestAttributeSet>()->Health;
 
-		Test->TestTrue(TEXT("Buff Instant Damage Applied"), (DestComponent->GetSet<USkillSystemTestAttributeSet>()->Health == ExpectedValue));
-		SKILL_LOG(Log, TEXT("Final Health: %.2f"), DestComponent->GetSet<USkillSystemTestAttributeSet>()->Health);
+		Test->TestTrue(SKILL_TEST_TEXT("Buff Instant Damage Applied.  Actual: %.2f == Exected: %.2f", ActualValue, ExpectedValue), (ActualValue == ExpectedValue));
+	}
+
+	return true;
+}
+
+
+bool GameplayEffectsTest_InstantDamage_ScalingProperty(UWorld *World, FAutomationTestBase * Test)
+{
+	// This example we scale Damage based off the instigator's PhysicalDamage attribute
+
+	const float StartHealth = 100.f;
+	const float PhysicalDamage = 10.f;
+	const float GameplayEffectScaling = 1.f;
+
+	SetGlobalCurveTable();
+
+	ASkillSystemTestPawn *SourceActor = World->SpawnActor<ASkillSystemTestPawn>();
+	ASkillSystemTestPawn *DestActor = World->SpawnActor<ASkillSystemTestPawn>();
+
+	UProperty *DamageProperty = FindFieldChecked<UProperty>(USkillSystemTestAttributeSet::StaticClass(), GET_MEMBER_NAME_CHECKED(USkillSystemTestAttributeSet, Damage));
+	UProperty *PhysicalDamageProperty = FindFieldChecked<UProperty>(USkillSystemTestAttributeSet::StaticClass(), GET_MEMBER_NAME_CHECKED(USkillSystemTestAttributeSet, PhysicalDamage));
+
+	UAttributeComponent * SourceComponent = SourceActor->AttributeComponent;
+	UAttributeComponent * DestComponent = DestActor->AttributeComponent;
+	SourceComponent->GetSet<USkillSystemTestAttributeSet>()->Health = StartHealth;
+	DestComponent->GetSet<USkillSystemTestAttributeSet>()->Health = StartHealth;
+
+	SourceComponent->GetSet<USkillSystemTestAttributeSet>()->PhysicalDamage = PhysicalDamage;
+
+	{
+		SKILL_LOG_SCOPE(TEXT("Apply InstantDamage"));
+
+		// This effects do Damage = 1.f * LinearCurve[LevelOfGameplayEffect].
+		// This translate into 1.f * PhysicalDamage.
+
+		UGameplayEffect * BaseDmgEffect = Cast<UGameplayEffect>(StaticConstructObject(UGameplayEffect::StaticClass(), GetTransientPackage(), FName(TEXT("BaseDmgEffect"))));
+		BaseDmgEffect->Modifiers.SetNum(1);
+		BaseDmgEffect->Modifiers[0].Magnitude.SetScalingValue(GameplayEffectScaling, FName(TEXT("LinearCurve")), NULL); // do "1*StandardDamage[Level]"
+		BaseDmgEffect->Modifiers[0].ModifierType = EGameplayMod::Attribute;
+		BaseDmgEffect->Modifiers[0].ModifierOp = EGameplayModOp::Additive;
+		BaseDmgEffect->Modifiers[0].Attribute.SetUProperty(DamageProperty);
+		BaseDmgEffect->Modifiers[0].OwnedTags.AddTag(FName(TEXT("Damage.Basic")));
+		BaseDmgEffect->Duration.Value = UGameplayEffect::INSTANT_APPLICATION;		
+		BaseDmgEffect->LevelInfo.Attribute.SetUProperty(PhysicalDamageProperty);
+		BaseDmgEffect->LevelInfo.InheritLevelFromOwner = false;
+		
+
+		SourceComponent->ApplyGameplayEffectSpecToTarget(BaseDmgEffect, DestComponent);
+
+		float ActualValue = DestComponent->GetSet<USkillSystemTestAttributeSet>()->Health;
+		float ExpectedValue = StartHealth - (PhysicalDamage * GameplayEffectScaling);
+		Test->TestTrue(SKILL_TEST_TEXT("Damage Applied. Health: Actual: %.2f == Exected: %.2f", ActualValue, ExpectedValue), ActualValue == ExpectedValue);
+	}
+
+	ClearGlobalCurveTable();
+	return true;
+}
+
+bool GameplayEffectsTest_InstantDamage_ScalingPropertyNested(UWorld *World, FAutomationTestBase * Test)
+{
+	// This accomplishes the same as GameplayEffectsTest_InstantDamage_ScalingProperty but the leveling info is specified at the modifier, not gameplayeffect, level.
+
+	const float StartHealth = 100.f;
+	const float PhysicalDamage = 10.f;
+	const float GameplayEffectScaling = 1.f;
+
+	SetGlobalCurveTable();
+
+	ASkillSystemTestPawn *SourceActor = World->SpawnActor<ASkillSystemTestPawn>();
+	ASkillSystemTestPawn *DestActor = World->SpawnActor<ASkillSystemTestPawn>();
+
+	UProperty *DamageProperty = FindFieldChecked<UProperty>(USkillSystemTestAttributeSet::StaticClass(), GET_MEMBER_NAME_CHECKED(USkillSystemTestAttributeSet, Damage));
+	UProperty *PhysicalDamageProperty = FindFieldChecked<UProperty>(USkillSystemTestAttributeSet::StaticClass(), GET_MEMBER_NAME_CHECKED(USkillSystemTestAttributeSet, PhysicalDamage));
+
+	UAttributeComponent * SourceComponent = SourceActor->AttributeComponent;
+	UAttributeComponent * DestComponent = DestActor->AttributeComponent;
+	SourceComponent->GetSet<USkillSystemTestAttributeSet>()->Health = StartHealth;
+	DestComponent->GetSet<USkillSystemTestAttributeSet>()->Health = StartHealth;
+
+	SourceComponent->GetSet<USkillSystemTestAttributeSet>()->PhysicalDamage = PhysicalDamage;
+
+	{
+		SKILL_LOG_SCOPE(TEXT("Apply InstantDamage"));
+
+		// This effects do Damage = 1.f * LinearCurve[LevelOfGameplayEffect].
+		// This translate into 1.f * PhysicalDamage.
+
+		UGameplayEffect * BaseDmgEffect = Cast<UGameplayEffect>(StaticConstructObject(UGameplayEffect::StaticClass(), GetTransientPackage(), FName(TEXT("BaseDmgEffect"))));
+		BaseDmgEffect->Modifiers.SetNum(1);
+		BaseDmgEffect->Modifiers[0].Magnitude.SetScalingValue(GameplayEffectScaling, FName(TEXT("LinearCurve")), NULL); // do "1*StandardDamage[Level]"
+		BaseDmgEffect->Modifiers[0].ModifierType = EGameplayMod::Attribute;
+		BaseDmgEffect->Modifiers[0].ModifierOp = EGameplayModOp::Additive;
+		BaseDmgEffect->Modifiers[0].Attribute.SetUProperty(DamageProperty);
+		BaseDmgEffect->Modifiers[0].OwnedTags.AddTag(FName(TEXT("Damage.Basic")));
+		BaseDmgEffect->Modifiers[0].LevelInfo.Attribute.SetUProperty(PhysicalDamageProperty);
+		BaseDmgEffect->Modifiers[0].LevelInfo.InheritLevelFromOwner = false;
+		BaseDmgEffect->Duration.Value = UGameplayEffect::INSTANT_APPLICATION;
+
+		SourceComponent->ApplyGameplayEffectSpecToTarget(BaseDmgEffect, DestComponent);
+
+		float ActualValue = DestComponent->GetSet<USkillSystemTestAttributeSet>()->Health;
+		float ExpectedValue = StartHealth - (PhysicalDamage * GameplayEffectScaling);
+		Test->TestTrue(SKILL_TEST_TEXT("Damage Applied. Health: Actual: %.2f == Exected: %.2f", ActualValue, ExpectedValue), ActualValue == ExpectedValue);
+	}
+
+	ClearGlobalCurveTable();
+	return true;
+}
+
+bool GameplayEffectsTest_DotDamage_ScalingProperty_Snapshot(UWorld *World, FAutomationTestBase * Test)
+{
+	// Add a dot that is powered by SpellDamage. Increase SpellDamage after applying, confirm it doesn't add extra damage to subsequent ticks.
+
+	const float StartHealth = 100.f;
+	const float SpellDamage = 10.f;
+	const float SpellDamage2 = 50.f;
+	const float GameplayEffectScaling = 1.f;
+
+	SetGlobalCurveTable();
+
+	ASkillSystemTestPawn *SourceActor = World->SpawnActor<ASkillSystemTestPawn>();
+	ASkillSystemTestPawn *DestActor = World->SpawnActor<ASkillSystemTestPawn>();
+
+	UProperty *DamageProperty = FindFieldChecked<UProperty>(USkillSystemTestAttributeSet::StaticClass(), GET_MEMBER_NAME_CHECKED(USkillSystemTestAttributeSet, Damage));
+	UProperty *SpellDamageProperty = FindFieldChecked<UProperty>(USkillSystemTestAttributeSet::StaticClass(), GET_MEMBER_NAME_CHECKED(USkillSystemTestAttributeSet, SpellDamage));
+
+	UAttributeComponent * SourceComponent = SourceActor->AttributeComponent;
+	UAttributeComponent * DestComponent = DestActor->AttributeComponent;
+	SourceComponent->GetSet<USkillSystemTestAttributeSet>()->Health = StartHealth;
+	DestComponent->GetSet<USkillSystemTestAttributeSet>()->Health = StartHealth;
+
+	SourceComponent->GetSet<USkillSystemTestAttributeSet>()->SpellDamage = SpellDamage;
+
+	{
+		SKILL_LOG_SCOPE(TEXT("Apply InstantDamage"));
+
+		// This effects do Damage = 1.f * LinearCurve[LevelOfGameplayEffect].
+		// This translate into 1.f * PhysicalDamage.
+
+		UGameplayEffect * BaseDmgEffect = Cast<UGameplayEffect>(StaticConstructObject(UGameplayEffect::StaticClass(), GetTransientPackage(), FName(TEXT("BaseDmgEffect"))));
+		BaseDmgEffect->Modifiers.SetNum(1);
+		BaseDmgEffect->Modifiers[0].Magnitude.SetScalingValue(GameplayEffectScaling, FName(TEXT("LinearCurve")), NULL); // do "1*StandardDamage[Level]"
+		BaseDmgEffect->Modifiers[0].ModifierType = EGameplayMod::Attribute;
+		BaseDmgEffect->Modifiers[0].ModifierOp = EGameplayModOp::Additive;
+		BaseDmgEffect->Modifiers[0].Attribute.SetUProperty(DamageProperty);				// Modifies target's "Damage" attribute (-health)
+		BaseDmgEffect->Modifiers[0].OwnedTags.AddTag(FName(TEXT("Damage.Basic")));
+		BaseDmgEffect->Duration.Value = UGameplayEffect::INFINITE_DURATION;
+		BaseDmgEffect->Period.Value = 1.f;
+
+		BaseDmgEffect->LevelInfo.Attribute.SetUProperty(SpellDamageProperty);			// Powered by instigators SpellDamage
+		BaseDmgEffect->LevelInfo.InheritLevelFromOwner = false;
+		BaseDmgEffect->LevelInfo.TakeSnapshotOnInit = true;								// But just a snapshot of their SpellDamage when we are applied
+
+		SourceComponent->ApplyGameplayEffectSpecToTarget(BaseDmgEffect, DestComponent);
+
+		// Confirm we ticked once
+		float ActualValue = DestComponent->GetSet<USkillSystemTestAttributeSet>()->Health;
+		float ExpectedValue = StartHealth - (SpellDamage * GameplayEffectScaling);
+		Test->TestTrue(SKILL_TEST_TEXT("Damage Applied. Health: Actual: %.2f == Exected: %.2f", ActualValue, ExpectedValue), ActualValue == ExpectedValue);
+	}
+
+
+	{
+		// Increase spell damage on instigator (after we already applied the DOT)
+		UGameplayEffect * SpellDmgEffect = Cast<UGameplayEffect>(StaticConstructObject(UGameplayEffect::StaticClass(), GetTransientPackage(), FName(TEXT("SpellDmgEffect"))));
+		SpellDmgEffect->Modifiers.SetNum(1);
+		SpellDmgEffect->Modifiers[0].Magnitude.SetValue(SpellDamage2);
+		SpellDmgEffect->Modifiers[0].ModifierType = EGameplayMod::Attribute;
+		SpellDmgEffect->Modifiers[0].ModifierOp = EGameplayModOp::Override;
+		SpellDmgEffect->Modifiers[0].Attribute.SetUProperty(SpellDamageProperty);				// Modifies target's "Damage" attribute (-health)
+		SpellDmgEffect->Modifiers[0].OwnedTags.AddTag(FName(TEXT("SpellDmg.Buff")));
+		SpellDmgEffect->Duration.Value = UGameplayEffect::INSTANT_APPLICATION;
+
+		SourceComponent->ApplyGameplayEffectSpecToTarget(SpellDmgEffect, SourceComponent);
+	}
+
+	{
+		// Tick once
+		GameplayTest_TickWorld(World, 1.f);
+		
+		float ActualValue = DestComponent->GetSet<USkillSystemTestAttributeSet>()->Health;
+		float ExpectedValue = StartHealth - (SpellDamage * GameplayEffectScaling) * 2.f;	// We've ticked twice
+		Test->TestTrue(SKILL_TEST_TEXT("Damage Applied. Health: Actual: %.2f == Exected: %.2f", ActualValue, ExpectedValue), ActualValue == ExpectedValue);
+	}	
+
+	ClearGlobalCurveTable();
+	return true;
+}
+
+bool GameplayEffectsTest_DotDamage_ScalingProperty_Dynamic(UWorld *World, FAutomationTestBase * Test)
+{
+	// Add a dot that is powered by SpellDamage. Increase SpellDamage after applying, confirm it doesn't add extra damage to subsequent ticks.
+
+	const float StartHealth = 100.f;
+	const float SpellDamage = 10.f;
+	const float SpellDamage2 = 50.f;
+	const float GameplayEffectScaling = 1.f;
+
+	SetGlobalCurveTable();
+
+	ASkillSystemTestPawn *SourceActor = World->SpawnActor<ASkillSystemTestPawn>();
+	ASkillSystemTestPawn *DestActor = World->SpawnActor<ASkillSystemTestPawn>();
+
+	UProperty *DamageProperty = FindFieldChecked<UProperty>(USkillSystemTestAttributeSet::StaticClass(), GET_MEMBER_NAME_CHECKED(USkillSystemTestAttributeSet, Damage));
+	UProperty *SpellDamageProperty = FindFieldChecked<UProperty>(USkillSystemTestAttributeSet::StaticClass(), GET_MEMBER_NAME_CHECKED(USkillSystemTestAttributeSet, SpellDamage));
+
+	UAttributeComponent * SourceComponent = SourceActor->AttributeComponent;
+	UAttributeComponent * DestComponent = DestActor->AttributeComponent;
+	SourceComponent->GetSet<USkillSystemTestAttributeSet>()->Health = StartHealth;
+	DestComponent->GetSet<USkillSystemTestAttributeSet>()->Health = StartHealth;
+
+	SourceComponent->GetSet<USkillSystemTestAttributeSet>()->SpellDamage = SpellDamage;
+
+	{
+		SKILL_LOG_SCOPE(TEXT("Apply InstantDamage"));
+
+		// This effects do Damage = 1.f * LinearCurve[LevelOfGameplayEffect].
+		// This translate into 1.f * PhysicalDamage.
+
+		UGameplayEffect * BaseDmgEffect = Cast<UGameplayEffect>(StaticConstructObject(UGameplayEffect::StaticClass(), GetTransientPackage(), FName(TEXT("BaseDmgEffect"))));
+		BaseDmgEffect->Modifiers.SetNum(1);
+		BaseDmgEffect->Modifiers[0].Magnitude.SetScalingValue(GameplayEffectScaling, FName(TEXT("LinearCurve")), NULL); // do "1*StandardDamage[Level]"
+		BaseDmgEffect->Modifiers[0].ModifierType = EGameplayMod::Attribute;
+		BaseDmgEffect->Modifiers[0].ModifierOp = EGameplayModOp::Additive;
+		BaseDmgEffect->Modifiers[0].Attribute.SetUProperty(DamageProperty);				// Modifies target's "Damage" attribute (-health)
+		BaseDmgEffect->Modifiers[0].OwnedTags.AddTag(FName(TEXT("Damage.Basic")));
+		BaseDmgEffect->Duration.Value = UGameplayEffect::INFINITE_DURATION;
+		BaseDmgEffect->Period.Value = 1.f;
+
+		BaseDmgEffect->LevelInfo.Attribute.SetUProperty(SpellDamageProperty);			// Powered by instigators SpellDamage
+		BaseDmgEffect->LevelInfo.InheritLevelFromOwner = false;
+		BaseDmgEffect->LevelInfo.TakeSnapshotOnInit = false;							// But level is dynamic, if SpellDamage changers after we apply, we update.
+
+		SourceComponent->ApplyGameplayEffectSpecToTarget(BaseDmgEffect, DestComponent);
+
+		// Confirm we ticked once
+		float ActualValue = DestComponent->GetSet<USkillSystemTestAttributeSet>()->Health;
+		float ExpectedValue = StartHealth - (SpellDamage * GameplayEffectScaling);
+		Test->TestTrue(SKILL_TEST_TEXT("Damage Applied. Health: Actual: %.2f == Exected: %.2f", ActualValue, ExpectedValue), ActualValue == ExpectedValue);
+
+		float SpellDamageTest = SourceComponent->GetSet<USkillSystemTestAttributeSet>()->SpellDamage;
+		check(SpellDamageTest == SpellDamage);
+	}
+
+	{
+		// Increase spell damage on instigator (after we already applied the DOT)
+		UGameplayEffect * SpellDmgEffect = Cast<UGameplayEffect>(StaticConstructObject(UGameplayEffect::StaticClass(), GetTransientPackage(), FName(TEXT("SpellDmgEffect"))));
+		SpellDmgEffect->Modifiers.SetNum(1);
+		SpellDmgEffect->Modifiers[0].Magnitude.SetValue(SpellDamage2);
+		SpellDmgEffect->Modifiers[0].ModifierType = EGameplayMod::Attribute;
+		SpellDmgEffect->Modifiers[0].ModifierOp = EGameplayModOp::Override;
+		SpellDmgEffect->Modifiers[0].Attribute.SetUProperty(SpellDamageProperty);				// Modifies target's "Damage" attribute (-health)
+		SpellDmgEffect->Modifiers[0].OwnedTags.AddTag(FName(TEXT("SpellDmg.Buff")));
+		SpellDmgEffect->Duration.Value = UGameplayEffect::INSTANT_APPLICATION;
+
+		SourceComponent->ApplyGameplayEffectSpecToTarget(SpellDmgEffect, SourceComponent);
+
+		float ActualValue = SourceComponent->GetSet<USkillSystemTestAttributeSet>()->SpellDamage;
+		float ExpectedValue = SpellDamage2;
+		Test->TestTrue(SKILL_TEST_TEXT("Spell Damage Mod: Actual: %.2f == Exected: %.2f", ActualValue, ExpectedValue), ActualValue == ExpectedValue);
+	}
+
+	{
+		// Tick once
+		GameplayTest_TickWorld(World, 1.f);
+
+		float ActualValue = DestComponent->GetSet<USkillSystemTestAttributeSet>()->Health;
+		float ExpectedValue = StartHealth - (SpellDamage * GameplayEffectScaling) - (SpellDamage2 * GameplayEffectScaling);	// We've ticked twice
+		Test->TestTrue(SKILL_TEST_TEXT("Damage Applied. Health: Actual: %.2f == Exected: %.2f", ActualValue, ExpectedValue), ActualValue == ExpectedValue);
+	}
+
+	ClearGlobalCurveTable();
+	return true;
+}
+
+bool GameplayEffectsTest_MetaAttributes(UWorld *World, FAutomationTestBase * Test)
+{
+	// Sets up a GameplayEffect to give the source a constant +Health powered by the source's strength
+
+	const float StartHealth = 100.f;
+	const float MaxHealthPerStrength = 3.f;
+	const float StrengthValue = 10.f;
+
+	SetGlobalCurveTable();
+
+	ASkillSystemTestPawn *SourceActor = World->SpawnActor<ASkillSystemTestPawn>();
+	
+
+	UProperty *MaxHealthProperty = FindFieldChecked<UProperty>(USkillSystemTestAttributeSet::StaticClass(), GET_MEMBER_NAME_CHECKED(USkillSystemTestAttributeSet, MaxHealth));
+	UProperty *StrengthProperty = FindFieldChecked<UProperty>(USkillSystemTestAttributeSet::StaticClass(), GET_MEMBER_NAME_CHECKED(USkillSystemTestAttributeSet, Strength));
+
+	UAttributeComponent * SourceComponent = SourceActor->AttributeComponent;
+	SourceComponent->GetSet<USkillSystemTestAttributeSet>()->MaxHealth = StartHealth;
+	SourceComponent->GetSet<USkillSystemTestAttributeSet>()->Strength = 0.f;
+
+	{
+		SKILL_LOG_SCOPE(TEXT("Setup meta stat"));
+
+		// This effects do Damage = 1.f * LinearCurve[LevelOfGameplayEffect].
+		// This translate into 1.f * PhysicalDamage.
+
+		UGameplayEffect * StrengthMaxHealhEffect = Cast<UGameplayEffect>(StaticConstructObject(UGameplayEffect::StaticClass(), GetTransientPackage(), FName(TEXT("StrengthMaxHealhEffect"))));
+		StrengthMaxHealhEffect->Modifiers.SetNum(1);
+		StrengthMaxHealhEffect->Modifiers[0].Magnitude.SetScalingValue(MaxHealthPerStrength, FName(TEXT("LinearCurve")), NULL); // do "1*StandardDamage[Level]"
+		StrengthMaxHealhEffect->Modifiers[0].ModifierType = EGameplayMod::Attribute;
+		StrengthMaxHealhEffect->Modifiers[0].ModifierOp = EGameplayModOp::Additive;
+		StrengthMaxHealhEffect->Modifiers[0].Attribute.SetUProperty(MaxHealthProperty);				// Modifies target's "Damage" attribute (-health)
+		StrengthMaxHealhEffect->Modifiers[0].OwnedTags.AddTag(FName(TEXT("Damage.Basic")));
+		StrengthMaxHealhEffect->Duration.Value = UGameplayEffect::INFINITE_DURATION;
+		StrengthMaxHealhEffect->Period.Value = UGameplayEffect::NO_PERIOD;
+
+		StrengthMaxHealhEffect->LevelInfo.Attribute.SetUProperty(StrengthProperty);			// Powered by instigators SpellDamage
+		StrengthMaxHealhEffect->LevelInfo.InheritLevelFromOwner = false;
+		StrengthMaxHealhEffect->LevelInfo.TakeSnapshotOnInit = false;						// But level is dynamic, if SpellDamage changers after we apply, we update.
+
+		SourceComponent->ApplyGameplayEffectSpecToTarget(StrengthMaxHealhEffect, SourceComponent);
+
+		// Strength starts at 0, so confirm it did nothing yet.
+		float ActualValue = SourceComponent->GetSet<USkillSystemTestAttributeSet>()->MaxHealth;
+		float ExpectedValue = StartHealth;
+		Test->TestTrue(SKILL_TEST_TEXT("Damage Applied. Health: Actual: %.2f == Exected: %.2f", ActualValue, ExpectedValue), ActualValue == ExpectedValue);
+	}
+
+	{
+		// Set strength to 10. Confirm this adds 30 to MaxHeatlh.
+		UGameplayEffect * StrEffect = Cast<UGameplayEffect>(StaticConstructObject(UGameplayEffect::StaticClass(), GetTransientPackage(), FName(TEXT("StrEffect"))));
+		StrEffect->Modifiers.SetNum(1);
+		StrEffect->Modifiers[0].Magnitude.SetValue(StrengthValue);
+		StrEffect->Modifiers[0].ModifierType = EGameplayMod::Attribute;
+		StrEffect->Modifiers[0].ModifierOp = EGameplayModOp::Additive;
+		StrEffect->Modifiers[0].Attribute.SetUProperty(StrengthProperty);				// Modifies target's "Damage" attribute (-health)
+		StrEffect->Modifiers[0].OwnedTags.AddTag(FName(TEXT("SpellDmg.Buff")));
+		StrEffect->Duration.Value = UGameplayEffect::INSTANT_APPLICATION;
+
+		SourceComponent->ApplyGameplayEffectSpecToTarget(StrEffect, SourceComponent);
+
+		{
+			float ActualValue = SourceComponent->GetSet<USkillSystemTestAttributeSet>()->Strength;
+			float ExpectedValue = StrengthValue;
+			Test->TestTrue(SKILL_TEST_TEXT("Strength: Actual: %.2f == Exected: %.2f", ActualValue, ExpectedValue), ActualValue == ExpectedValue);
+		}
+
+		{
+			float ActualValue = SourceComponent->GetSet<USkillSystemTestAttributeSet>()->MaxHealth;
+			float ExpectedValue = StartHealth + (MaxHealthPerStrength * StrengthValue);
+			Test->TestTrue(SKILL_TEST_TEXT("MaxHealth: Actual: %.2f == Exected: %.2f", ActualValue, ExpectedValue), ActualValue == ExpectedValue);
+		}
+	}
+
+	ClearGlobalCurveTable();
+	return true;
+}
+
+bool GameplayEffectsTest_TagOrdering(UWorld *World, FAutomationTestBase * Test)
+{
+	const float StartHealth = 100.f;
+	const float DamageValue = 5.f;
+	const float BonusDamageMultiplier = 2.f;
+
+	ASkillSystemTestPawn *SourceActor = World->SpawnActor<ASkillSystemTestPawn>();
+	ASkillSystemTestPawn *DestActor = World->SpawnActor<ASkillSystemTestPawn>();
+
+	UProperty *HealthProperty = FindFieldChecked<UProperty>(USkillSystemTestAttributeSet::StaticClass(), GET_MEMBER_NAME_CHECKED(USkillSystemTestAttributeSet, Health));
+	UProperty *DamageProperty = FindFieldChecked<UProperty>(USkillSystemTestAttributeSet::StaticClass(), GET_MEMBER_NAME_CHECKED(USkillSystemTestAttributeSet, Damage));
+
+	UAttributeComponent * SourceComponent = SourceActor->AttributeComponent;
+	UAttributeComponent * DestComponent = DestActor->AttributeComponent;
+	SourceComponent->GetSet<USkillSystemTestAttributeSet>()->Health = StartHealth;
+	DestComponent->GetSet<USkillSystemTestAttributeSet>()->Health = StartHealth;
+
+	// Setup a GE to modify OutgoingGEs
+	{
+		SKILL_LOG_SCOPE(TEXT("FireDamageBuff"));
+
+		UGameplayEffect * BaseDmgEffect = Cast<UGameplayEffect>(StaticConstructObject(UGameplayEffect::StaticClass(), GetTransientPackage(), FName(TEXT("FireDamageBuff"))));
+		BaseDmgEffect->Modifiers.SetNum(1);
+		BaseDmgEffect->Modifiers[0].Magnitude.SetValue(BonusDamageMultiplier);
+		BaseDmgEffect->Modifiers[0].ModifierType = EGameplayMod::OutgoingGE;
+		BaseDmgEffect->Modifiers[0].ModifierOp = EGameplayModOp::Multiplicitive;
+		BaseDmgEffect->Modifiers[0].Attribute.SetUProperty(DamageProperty);
+		BaseDmgEffect->Modifiers[0].OwnedTags.AddTag(FName(TEXT("Damage.Buffed.FireBuff")));
+		BaseDmgEffect->Modifiers[0].RequiredTags.AddTag(FName(TEXT("Damage.Fire")));
+		BaseDmgEffect->Duration.SetValue(UGameplayEffect::INFINITE_DURATION);
+		
+		BaseDmgEffect->GameplayEffectTags.AddTag(FName(TEXT("Buff")));
+		BaseDmgEffect->GameplayEffectRequiredTags.AddTag(FName(TEXT("Damage")));
+
+		// Apply to self
+		SourceComponent->ApplyGameplayEffectSpecToTarget(BaseDmgEffect, SourceComponent, 1.f);
+	}
+
+	{
+		SKILL_LOG_SCOPE(TEXT("MakeFireDamage"));
+
+		UGameplayEffect * BaseDmgEffect = Cast<UGameplayEffect>(StaticConstructObject(UGameplayEffect::StaticClass(), GetTransientPackage(), FName(TEXT("MakeFireDamage"))));
+		BaseDmgEffect->Modifiers.SetNum(1);
+		BaseDmgEffect->Modifiers[0].Magnitude.SetValue(0.f);
+		BaseDmgEffect->Modifiers[0].ModifierType = EGameplayMod::OutgoingGE;
+		BaseDmgEffect->Modifiers[0].ModifierOp = EGameplayModOp::Additive;
+		BaseDmgEffect->Modifiers[0].Attribute.SetUProperty(DamageProperty);
+		BaseDmgEffect->Modifiers[0].OwnedTags.AddTag(FName(TEXT("Damage.Fire")));
+		BaseDmgEffect->Modifiers[0].RequiredTags.AddTag(FName(TEXT("Damage.Physical")));
+		BaseDmgEffect->Duration.SetValue(UGameplayEffect::INFINITE_DURATION);
+		
+		BaseDmgEffect->GameplayEffectTags.AddTag(FName(TEXT("Buff")));
+		BaseDmgEffect->GameplayEffectRequiredTags.AddTag(FName(TEXT("Damage")));
+
+		// Apply to self
+		SourceComponent->ApplyGameplayEffectSpecToTarget(BaseDmgEffect, SourceComponent, 1.f);
+	}
+
+	// Apply Damage
+	{
+		SKILL_LOG_SCOPE(TEXT("Apply InstantDamage"));
+
+		UGameplayEffect * BaseDmgEffect = Cast<UGameplayEffect>(StaticConstructObject(UGameplayEffect::StaticClass(), GetTransientPackage(), FName(TEXT("Damage"))));
+		BaseDmgEffect->Modifiers.SetNum(1);
+		BaseDmgEffect->Modifiers[0].Magnitude.SetValue(DamageValue);
+		BaseDmgEffect->Modifiers[0].ModifierType = EGameplayMod::Attribute;
+		BaseDmgEffect->Modifiers[0].ModifierOp = EGameplayModOp::Additive;
+		BaseDmgEffect->Modifiers[0].Attribute.SetUProperty(DamageProperty);
+		BaseDmgEffect->Modifiers[0].OwnedTags.AddTag(FName(TEXT("Damage.Physical")));
+		BaseDmgEffect->Duration.Value = UGameplayEffect::INSTANT_APPLICATION;
+		
+		BaseDmgEffect->GameplayEffectTags.AddTag(FName(TEXT("Damage")));
+
+		// Apply to target
+		SourceComponent->ApplyGameplayEffectSpecToTarget(BaseDmgEffect, DestComponent, 1.f);
+
+		float ExpectedValue = (StartHealth - (DamageValue * BonusDamageMultiplier));
+		float ActualValue = DestComponent->GetSet<USkillSystemTestAttributeSet>()->Health;
+		
+		Test->TestTrue(SKILL_TEST_TEXT("MaxHealth: Actual: %.2f == Exected: %.2f", ActualValue, ExpectedValue), ActualValue == ExpectedValue);
 	}
 
 	return true;
@@ -2305,7 +2737,7 @@ bool FGameplayEffectsTest::RunTest( const FString& Parameters )
 	UWorld *World = UWorld::CreateWorld( EWorldType::Game, false );
 	FURL URL;
 	World->BeginPlay(URL);
-	
+
 	GameplayEffectsTest_InstantDamage(World, this);
 	GameplayEffectsTest_InstantDamageRemap(World, this);
 	GameplayEffectsTest_InstantDamage_Buffed(World, this);
@@ -2336,8 +2768,18 @@ bool FGameplayEffectsTest::RunTest( const FString& Parameters )
 
 	GameplayEffectsTest_InstantDamageRequiredTag(World, this);
 	GameplayEffectsTest_InstantDamageIgnoreTag(World, this);
+	
 	GameplayEffectsTest_InstantDamageModifierPassesTag(World, this);
 	GameplayEffectsTest_InstantDamageModifierTag(World, this);
+	GameplayEffectsTest_InstantDamage_ScalingProperty(World, this);
+	GameplayEffectsTest_InstantDamage_ScalingPropertyNested(World, this);
+
+	GameplayEffectsTest_DotDamage_ScalingProperty_Snapshot(World, this);
+	GameplayEffectsTest_DotDamage_ScalingProperty_Dynamic(World, this);	
+
+	GameplayEffectsTest_MetaAttributes(World, this);
+	GameplayEffectsTest_TagOrdering(World, this);
+	
 
 	World->DestroyWorld(false);
 
diff --git a/Engine/Source/Runtime/SkillSystem/Private/SkillSystemTestAttributeSet.cpp b/Engine/Source/Runtime/SkillSystem/Private/SkillSystemTestAttributeSet.cpp
--- a/Engine/Source/Runtime/SkillSystem/Private/SkillSystemTestAttributeSet.cpp
+++ b/Engine/Source/Runtime/SkillSystem/Private/SkillSystemTestAttributeSet.cpp
@@ -8,6 +8,9 @@ USkillSystemTestAttributeSet::USkillSystemTestAttributeSet(const class FPostCons
 	Health = 100.f;
 	Damage = 0.f;
 	CritChance = 0.f;
+	SpellDamage = 0.f;
+	PhysicalDamage = 0.f;
+	Strength = 0.f;
 }
 
 
diff --git a/Engine/Source/Runtime/SkillSystem/Private/SkillSystemTestPawn.cpp b/Engine/Source/Runtime/SkillSystem/Private/SkillSystemTestPawn.cpp
--- a/Engine/Source/Runtime/SkillSystem/Private/SkillSystemTestPawn.cpp
+++ b/Engine/Source/Runtime/SkillSystem/Private/SkillSystemTestPawn.cpp
@@ -18,7 +18,6 @@ void ASkillSystemTestPawn::PostInitializeComponents()
 	Super::PostInitializeComponents();
 	GameplayCueHandler.Owner = this;
 	AttributeComponent->InitStats(USkillSystemTestAttributeSet::StaticClass(), NULL);
-	AttributeComponent->RegisterAttributeModifyCallback(FGameplayAttribute(DamageProperty), FOnGameplayAttributeEffectExecuted::CreateUObject(this, &ASkillSystemTestPawn::OnDamageExecute));
 }
 
 void ASkillSystemTestPawn::GameplayCueActivated(const FGameplayTagContainer & GameplayCueTags, float NormalizedMagnitude)
@@ -43,37 +42,4 @@ void ASkillSystemTestPawn::GameplayCueRemoved(const FGameplayTagContainer & Game
 {
 	SKILL_LOG(Log, TEXT("GameplayCueRemoved: %s. %.2f"), *GetName(), NormalizedMagnitude);
 	GameplayCueHandler.GameplayCueRemoved(GameplayCueTags, NormalizedMagnitude);
-}
-
-void ASkillSystemTestPawn::OnDamageExecute(FGameplayModifierEvaluatedData & Data)
-{
-	auto AttributeSet = AttributeComponent->GetSet<USkillSystemTestAttributeSet>();
-
-	// A very basic example of how you could do a critical chance roll at the actor level
-	if (Data.Tags.HasTag(FName(TEXT("CanCrit"))))
-	{
-		if (AttributeSet->CritChance >= FMath::FRand())
-		{
-			// Crit success - modify Data
-			// (this example is simple - we will
-			Data.Magnitude *= AttributeSet->CritMultiplier;
-		}
-	}
-}
-
-void ASkillSystemTestPawn::OnDidDamage(const struct FGameplayEffectSpec & Spec, struct FGameplayModifierEvaluatedData & Data)
-{
-	auto AttributeSet = AttributeComponent->GetSet<USkillSystemTestAttributeSet>();
-
-	// A very basic example of how you could do a critical chance roll at the actor level
-	if (Data.Tags.HasTag(FName(TEXT("LifeSteal"))))
-	{
-		if (AttributeSet->LifeSteal >= 0.f)
-		{
-			float LifeStealAmount = AttributeSet->LifeSteal * Data.Magnitude;
-
-			// We should apply a gameplay effect to ourself instead
-			AttributeSet->Health += LifeStealAmount;
-		}
-	}
 }
\ No newline at end of file
