#SkillSystem
automated test fix

diff --git a/Engine/Source/Runtime/SkillSystem/Private/GameplayEffectTests.cpp b/Engine/Source/Runtime/SkillSystem/Private/GameplayEffectTests.cpp
--- a/Engine/Source/Runtime/SkillSystem/Private/GameplayEffectTests.cpp
+++ b/Engine/Source/Runtime/SkillSystem/Private/GameplayEffectTests.cpp
@@ -2734,7 +2734,10 @@ bool FGameplayEffectsTest::RunTest( const FString& Parameters )
 {
 	UCurveTable *CurveTable = ISkillSystemModule::Get().GetSkillSystemGlobals().GetGlobalCurveTable();
 
-	UWorld *World = UWorld::CreateWorld( EWorldType::Game, false );
+	UWorld *World = UWorld::CreateWorld(EWorldType::Game, false);
+	FWorldContext &WorldContext = GEngine->CreateNewWorldContext(EWorldType::Game);
+	WorldContext.SetCurrentWorld(World);
+	
 	FURL URL;
 	World->BeginPlay(URL);
 
@@ -2781,6 +2784,7 @@ bool FGameplayEffectsTest::RunTest( const FString& Parameters )
 	GameplayEffectsTest_TagOrdering(World, this);
 	
 
+	GEngine->DestroyWorldContext(World);
 	World->DestroyWorld(false);
 
 	ISkillSystemModule::Get().GetSkillSystemGlobals().AutomationTestOnly_SetGlobalCurveTable(CurveTable);
