Add stacking rules for gameplay effects.

#ue4

 - Added stacking rules and tests.
 - Changed effects with a period to execute on the next tick instead of instantly and updated gameplay effect tests to account for these changes.
 - Changed ticking of active effects to allow effects to execute multiple times in a single tick.
 - Make sure we're cleaning up actors between gameplay effect tests

diff --git a/Engine/Source/Runtime/SkillSystem/Classes/AttributeSet.h b/Engine/Source/Runtime/SkillSystem/Classes/AttributeSet.h
--- a/Engine/Source/Runtime/SkillSystem/Classes/AttributeSet.h
+++ b/Engine/Source/Runtime/SkillSystem/Classes/AttributeSet.h
@@ -203,6 +203,9 @@ public:
 
 	UPROPERTY()
 	FString		DerivedAttributeInfo;
+
+	UPROPERTY()
+	bool		bCanStack;
 };
 
 
diff --git a/Engine/Source/Runtime/SkillSystem/Classes/GameplayEffect.h b/Engine/Source/Runtime/SkillSystem/Classes/GameplayEffect.h
--- a/Engine/Source/Runtime/SkillSystem/Classes/GameplayEffect.h
+++ b/Engine/Source/Runtime/SkillSystem/Classes/GameplayEffect.h
@@ -93,6 +93,22 @@ namespace EGameplayEffectCopyPolicy
 	};
 }
 
+UENUM(BlueprintType)
+namespace EGameplayEffectStackingPolicy
+{
+	enum Type
+	{
+		Unlimited = 0		UMETA(DisplayName = "NoRule"),
+		Highest				UMETA(DisplayName = "Strongest"),
+		Lowest				UMETA(DisplayName = "Weakest"),
+		Replaces			UMETA(DisplayName = "MostRecent"),
+		Callback			UMETA(DisplayName = "Custom"),
+
+		// This must always be at the end
+		Max					UMETA(DisplayName = "Invalid")
+	};
+}
+
 FString EGameplayModOpToString( int32 Type );
 
 FString EGameplayModToString( int32 Type );
@@ -101,6 +117,8 @@ FString EGameplayModEffectToString( int32 Type );
 
 FString EGameplayEffectCopyPolicyToString(int32 Type);
 
+FString EGameplayEffectStackingPolicyToString(int32 Type);
+
 USTRUCT()
 struct FGameplayModifierCallbacks
 {
@@ -110,6 +128,15 @@ struct FGameplayModifierCallbacks
 	TArray<TSubclassOf<class UGameplayEffectExtension> >	ExtensionClasses;
 };
 
+USTRUCT()
+struct FGameplayEffectStackingCallbacks
+{
+	GENERATED_USTRUCT_BODY()
+
+	UPROPERTY(EditDefaultsOnly, Category = GEStack)
+	TArray<TSubclassOf<class UGameplayEffectStackingExtension> >	ExtensionClasses;
+};
+
 /**
 * Defines how A GameplayEffect levels
 *	Normally, GameplayEffect levels are specified when they are created.
@@ -376,12 +403,20 @@ public:
 	/** Description of this combat effect. */
 	UPROPERTY(EditDefaultsOnly, BlueprintReadWrite, Category = Display)
 	FText Description;
+
+	// ----------------------------------------------
+
+	UPROPERTY(EditDefaultsOnly, BlueprintReadOnly, Category = Stacking)
+	TEnumAsByte<EGameplayEffectStackingPolicy::Type>	StackingPolicy;
+
+	UPROPERTY(EditDefaultsOnly, Category = Stacking)
+	TSubclassOf<class UGameplayEffectStackingExtension> StackingExtension;
 };
 
 /**
 * This handle is required for things outside of FActiveGameplayEffectsContainer to refer to a specific active GameplayEffect
 *	For example if a skill needs to create an active effect and then destroy that specific effect that it created, it has to do so
-*	through a handle. a pointer or index into the active list is not suffecient.
+*	through a handle. a pointer or index into the active list is not sufficient.
 */
 USTRUCT(BlueprintType)
 struct FActiveGameplayEffectHandle
@@ -960,7 +995,7 @@ private:
  *	-Const data (FGameplayModifierInfo) tells us what we modify, what we can modify
  *	-Mutable Aggregated data tells us how we modify (magnitude).
  *  
- * Modifiers can be modified. A modifier spec holds these modifications along with a refrence to the const data about the modifier.
+ * Modifiers can be modified. A modifier spec holds these modifications along with a reference to the const data about the modifier.
  * 
  */
 struct FModifierSpec
@@ -1014,6 +1049,7 @@ struct FGameplayEffectSpec
 		: ModifierLevel( TSharedPtr< FGameplayEffectLevelSpec >( new FGameplayEffectLevelSpec() ) )
 		, Duration(new FAggregator(FGameplayModifierEvaluatedData(0.f, NULL, FActiveGameplayEffectHandle()), SKILL_AGG_DEBUG(TEXT("Uninitialized Duration"))))
 		, Period(new FAggregator(FGameplayModifierEvaluatedData(0.f, NULL, FActiveGameplayEffectHandle()), SKILL_AGG_DEBUG(TEXT("Uninitialized Period"))))
+		, StackedAttribName(NAME_None)
 	{
 		// If we initialize a GameplayEffectSpec with no level object passed in.
 	}
@@ -1031,10 +1067,14 @@ struct FGameplayEffectSpec
 
 	float GetDuration() const;
 	float GetPeriod() const;
+	float GetMagnitude(const FGameplayAttribute &Attribute) const;
 
 	FAggregatorRef	Duration;
 	FAggregatorRef	Period;
 
+	EGameplayEffectStackingPolicy::Type StackingPolicy;
+	FName StackedAttribName;
+
 	// The spec needs to own these FModifierSpecs so that other people can keep TSharedPtr to it.
 	// The stuff in this array is OWNED by this spec
 
@@ -1049,6 +1089,8 @@ struct FGameplayEffectSpec
 
 	bool ShouldApplyAsSnapshot(const FModifierQualifier &QualifierContext) const;
 
+	EGameplayEffectStackingPolicy::Type GetStackingType() const;
+
 	FString ToSimpleString() const
 	{
 		return FString::Printf(TEXT("%s"), *Def->GetName());
@@ -1089,7 +1131,7 @@ struct FActiveGameplayEffect : public FFastArraySerializerItem
 		float Period = GetPeriod();
 		if (Period != UGameplayEffect::NO_PERIOD)
 		{
-			NextExecuteTime = CurrentTime + Period;
+			NextExecuteTime = CurrentTime + SMALL_NUMBER;
 		}
 
 		for (FModifierSpec &Mod : Spec.Modifiers)
@@ -1130,6 +1172,11 @@ struct FActiveGameplayEffect : public FFastArraySerializerItem
 	void PreReplicatedRemove(const struct FActiveGameplayEffectsContainer &InArray);
 	void PostReplicatedAdd(const struct FActiveGameplayEffectsContainer &InArray);
 	void PostReplicatedChange(const struct FActiveGameplayEffectsContainer &InArray) { }
+
+	bool operator==(const FActiveGameplayEffect& Other)
+	{
+		return Handle == Other.Handle;
+	}
 };
 
 /**
@@ -1137,7 +1184,7 @@ struct FActiveGameplayEffect : public FFastArraySerializerItem
  *	-Bucket of ActiveGameplayEffects
  *	-Needed for FFastArraySerialization
  *  
- * This should only be used by UAttributeComponent. All of this could just live in UAttributeComponent except that we need a distinc USTRUCT to implement FFastArraySerializer.
+ * This should only be used by UAttributeComponent. All of this could just live in UAttributeComponent except that we need a distinct USTRUCT to implement FFastArraySerializer.
  *
  */
 USTRUCT()
@@ -1145,8 +1192,10 @@ struct FActiveGameplayEffectsContainer : public FFastArraySerializer
 {
 	GENERATED_USTRUCT_BODY();
 
+	FActiveGameplayEffectsContainer() : bNeedToRecalculateStacks(false) {};
+
 	UPROPERTY()
-	TArray<FActiveGameplayEffect >	GameplayEffects;
+	TArray< FActiveGameplayEffect >	GameplayEffects;
 	
 	FActiveGameplayEffect & CreateNewActiveGameplayEffect(const FGameplayEffectSpec &Spec, float GameTimeSeconds);
 
@@ -1166,6 +1215,8 @@ struct FActiveGameplayEffectsContainer : public FFastArraySerializer
 
 	float GetGameplayEffectMagnitude(FActiveGameplayEffectHandle Handle, FGameplayAttribute Attribute) const;
 
+	// returns true if the handle points to an effect in this container that is not a stacking effect or an effect in this container that does stack and is applied by the current stacking rules
+	// returns false if the handle points to an effect that is not in this container or is not applied because of the current stacking rules
 	bool IsGameplayEffectActive(FActiveGameplayEffectHandle Handle) const;
 
 	void PrintAllGameplayEffects() const;
@@ -1183,6 +1234,9 @@ struct FActiveGameplayEffectsContainer : public FFastArraySerializer
 
 	void TEMP_TickActiveEffects(float DeltaSeconds);
 
+	// recalculates all of the stacks in the current container
+	void RecalculateStacking();
+
 	bool NetDeltaSerialize(FNetDeltaSerializeInfo & DeltaParms)
 	{
 		return FastArrayDeltaSerialize<FActiveGameplayEffect>(GameplayEffects, DeltaParms, *this);
@@ -1190,6 +1244,8 @@ struct FActiveGameplayEffectsContainer : public FFastArraySerializer
 
 	void PreDestroy();
 
+	bool bNeedToRecalculateStacks;
+
 private:
 
 	FAggregatorRef & FindOrCreateAttributeAggregator(FGameplayAttribute Attribute);
diff --git a/Engine/Source/Runtime/SkillSystem/Classes/GameplayEffectStackingExtension.h b/Engine/Source/Runtime/SkillSystem/Classes/GameplayEffectStackingExtension.h
new file mode 100644
--- /dev/null
+++ b/Engine/Source/Runtime/SkillSystem/Classes/GameplayEffectStackingExtension.h
@@ -0,0 +1,84 @@
+// Copyright 1998-2014 Epic Games, Inc. All Rights Reserved.
+
+#pragma once
+
+#include "GameplayTagContainer.h"
+#include "GameplayTagAssetInterface.h"
+#include "AttributeSet.h"
+#include "GameplayEffectStackingExtension.generated.h"
+
+struct FGameplayEffectStackingCallbackData
+{
+	FGameplayEffectStackingCallbackData(const FGameplayEffectSpec & InEffectSpec, const FModifierSpec & InModifierSpec, FGameplayModifierEvaluatedData & InEvaluatedData, UAttributeComponent & InTarget)
+		: EffectSpec(InEffectSpec)
+		, ModifierSpec(InModifierSpec)
+		, EvaluatedData(InEvaluatedData)
+		, Target(InTarget)
+	{
+
+	}
+
+	const struct FGameplayEffectSpec &		EffectSpec;		// The spec that the mod came from
+	const struct FModifierSpec &			ModifierSpec;	// The mod we are going to apply
+	struct FGameplayModifierEvaluatedData & EvaluatedData;	// The 'flat'/computed data to be applied to the target
+
+	class UAttributeComponent &Target;		// Target we intend to apply to
+};
+
+
+USTRUCT(BlueprintType)
+struct FGameplayEffectStackingHandle
+{
+	GENERATED_USTRUCT_BODY()
+
+	FGameplayEffectStackingHandle()
+	: Handle(INDEX_NONE)
+	{
+		// do nothing
+	}
+
+	FGameplayEffectStackingHandle(uint32 InHandle)
+		: Handle(InHandle)
+	{
+		// do nothing
+	}
+
+	bool IsValid() const
+	{
+		return Handle != INDEX_NONE;
+	}
+
+	bool operator==(const FGameplayEffectStackingHandle& Other) const
+	{
+		return Handle == Other.Handle;
+	}
+
+	bool operator!=(const FGameplayEffectStackingHandle& Other) const
+	{
+		return Handle != Other.Handle;
+	}
+
+	FString ToString() const
+	{
+		return FString::Printf(TEXT("%d"), Handle);
+	}
+
+private:
+
+	UPROPERTY()
+	uint32 Handle;
+};
+
+UCLASS(BlueprintType)
+class SKILLSYSTEM_API UGameplayEffectStackingExtension: public UObject
+{
+	GENERATED_UCLASS_BODY()
+
+public:
+
+	static FGameplayEffectSpec* StackCustomEffects(TArray<FGameplayEffectSpec*>& CustomGameplayEffects);
+
+	virtual void CalculateStack(TArray<FActiveGameplayEffect*>& CustomGameplayEffects, FActiveGameplayEffectsContainer& Container, FActiveGameplayEffect& CurrentEffect) {}
+
+	FGameplayEffectStackingHandle Handle;
+};
diff --git a/Engine/Source/Runtime/SkillSystem/Classes/GameplayEffectStackingExtension_CappedNumberTest.h b/Engine/Source/Runtime/SkillSystem/Classes/GameplayEffectStackingExtension_CappedNumberTest.h
new file mode 100644
--- /dev/null
+++ b/Engine/Source/Runtime/SkillSystem/Classes/GameplayEffectStackingExtension_CappedNumberTest.h
@@ -0,0 +1,18 @@
+// Copyright 1998-2014 Epic Games, Inc. All Rights Reserved.
+
+#pragma once
+
+#include "GameplayTagContainer.h"
+#include "GameplayTagAssetInterface.h"
+#include "AttributeSet.h"
+#include "GameplayEffectStackingExtension_CappedNumberTest.generated.h"
+
+UCLASS(BlueprintType)
+class SKILLSYSTEM_API UGameplayEffectStackingExtension_CappedNumberTest : public UGameplayEffectStackingExtension
+{
+	GENERATED_UCLASS_BODY()
+
+public:
+
+	void CalculateStack(TArray<FActiveGameplayEffect*>& CustomGameplayEffects, FActiveGameplayEffectsContainer& Container, FActiveGameplayEffect& CurrentEffect) OVERRIDE;
+};
\ No newline at end of file
diff --git a/Engine/Source/Runtime/SkillSystem/Classes/GameplayEffectStackingExtension_DiminishingReturnsTest.h b/Engine/Source/Runtime/SkillSystem/Classes/GameplayEffectStackingExtension_DiminishingReturnsTest.h
new file mode 100644
--- /dev/null
+++ b/Engine/Source/Runtime/SkillSystem/Classes/GameplayEffectStackingExtension_DiminishingReturnsTest.h
@@ -0,0 +1,18 @@
+// Copyright 1998-2014 Epic Games, Inc. All Rights Reserved.
+
+#pragma once
+
+#include "GameplayTagContainer.h"
+#include "GameplayTagAssetInterface.h"
+#include "AttributeSet.h"
+#include "GameplayEffectStackingExtension_DiminishingReturnsTest.generated.h"
+
+UCLASS(BlueprintType)
+class SKILLSYSTEM_API UGameplayEffectStackingExtension_DiminishingReturnsTest : public UGameplayEffectStackingExtension
+{
+	GENERATED_UCLASS_BODY()
+
+public:
+
+	void CalculateStack(TArray<FActiveGameplayEffect*>& CustomGameplayEffects, FActiveGameplayEffectsContainer& Container, FActiveGameplayEffect& CurrentEffect) OVERRIDE;
+};
\ No newline at end of file
diff --git a/Engine/Source/Runtime/SkillSystem/Classes/SkillSystemGlobals.h b/Engine/Source/Runtime/SkillSystem/Classes/SkillSystemGlobals.h
--- a/Engine/Source/Runtime/SkillSystem/Classes/SkillSystemGlobals.h
+++ b/Engine/Source/Runtime/SkillSystem/Classes/SkillSystemGlobals.h
@@ -14,19 +14,33 @@ class SKILLSYSTEM_API USkillSystemGlobals : public UObject
 	UPROPERTY(config)
 	FString GlobalCurveTableName;
 
+	/** Holds information about the valid attributes' min and max values and stacking rules */
+	UPROPERTY(config)
+	FString GlobalAttributeDataTableName;
+
 	class UCurveTable *	GetGlobalCurveTable();
 
+	class UDataTable * GetGlobalAttributeDataTable();
+
 	void AutomationTestOnly_SetGlobalCurveTable(class UCurveTable *InTable)
 	{
 		GlobalCurveTable = InTable;
 	}
 
+	void AutomationTestOnly_SetGlobalAttributeDataTable(class UDataTable *InTable)
+	{
+		GlobalAttributeDataTable = InTable;
+	}
+
 private:
 
 	class UCurveTable* GlobalCurveTable;
 
+	class UDataTable* GlobalAttributeDataTable;
+
 #if WITH_EDITOR
-	void OnObjectReimported(UObject* InObject);
+	void OnCurveTableReimported(UObject* InObject);
+	void OnDataTableReimported(UObject* InObject);
 #endif
 
 };
diff --git a/Engine/Source/Runtime/SkillSystem/Classes/SkillSystemTestAttributeSet.h b/Engine/Source/Runtime/SkillSystem/Classes/SkillSystemTestAttributeSet.h
--- a/Engine/Source/Runtime/SkillSystem/Classes/SkillSystemTestAttributeSet.h
+++ b/Engine/Source/Runtime/SkillSystem/Classes/SkillSystemTestAttributeSet.h
@@ -8,10 +8,10 @@ class USkillSystemTestAttributeSet : public UAttributeSet
 {
 	GENERATED_UCLASS_BODY()
 
-	UPROPERTY(EditAnywhere, BlueprintReadWrite, Category = "AttributeTest", meta = (HideFromModifiers))			// You can't make a GameplayEffect modifiy Health Directly (go through)
+	UPROPERTY(EditAnywhere, BlueprintReadWrite, Category = "AttributeTest", meta = (HideFromModifiers))			// You can't make a GameplayEffect modify Health Directly (go through)
 	float	MaxHealth;
 
-	UPROPERTY(EditAnywhere, BlueprintReadWrite, Category = "AttributeTest", meta = (HideFromModifiers))			// You can't make a GameplayEffect modifiy Health Directly (go through)
+	UPROPERTY(EditAnywhere, BlueprintReadWrite, Category = "AttributeTest", meta = (HideFromModifiers))			// You can't make a GameplayEffect modify Health Directly (go through)
 	float	Health;
 
 	/** This Damage is just used for applying negative health mods. Its not a 'persistent' attribute. */
@@ -22,7 +22,7 @@ class USkillSystemTestAttributeSet : public UAttributeSet
 	UPROPERTY(EditAnywhere, BlueprintReadWrite, Category = "AttributeTest")
 	float	SpellDamage;
 
-	/** This Attribute is the actual physical damage for this actor. It will power spell based GameplayEffects */
+	/** This Attribute is the actual physical damage for this actor. It will power physical based GameplayEffects */
 	UPROPERTY(EditAnywhere, BlueprintReadWrite, Category = "AttributeTest")
 	float	PhysicalDamage;
 
@@ -44,6 +44,15 @@ class USkillSystemTestAttributeSet : public UAttributeSet
 	UPROPERTY(EditAnywhere, BlueprintReadWrite, Category = "AttributeTest")
 	float	Strength;
 
+	UPROPERTY(EditAnywhere, BlueprintReadWrite, Category = "AttributeTest")
+	float	StackingAttribute1;
+
+	UPROPERTY(EditAnywhere, BlueprintReadWrite, Category = "AttributeTest")
+	float	StackingAttribute2;
+
+	UPROPERTY(EditAnywhere, BlueprintReadWrite, Category = "AttributeTest")
+	float	NoStackAttribute;
+
 	virtual void PreAttributeModify(struct FGameplayEffectModCallbackData &Data) OVERRIDE;
 	virtual void PostAttributeModify(const struct FGameplayEffectModCallbackData &Data) OVERRIDE;
 };
\ No newline at end of file
diff --git a/Engine/Source/Runtime/SkillSystem/Private/AttributeComponent.cpp b/Engine/Source/Runtime/SkillSystem/Private/AttributeComponent.cpp
--- a/Engine/Source/Runtime/SkillSystem/Private/AttributeComponent.cpp
+++ b/Engine/Source/Runtime/SkillSystem/Private/AttributeComponent.cpp
@@ -276,12 +276,15 @@ FActiveGameplayEffectHandle UAttributeComponent::ApplyGameplayEffectSpecToSelf(c
 		// have other modifiers. THOSE modifiers may or may not be copies of whatever.
 		//
 		// In theory, we don't modify 2nd order modifiers after they are 'attached'
-		// Long complex chains can be created but we never say 'Modify a GE that is moding another GE'
+		// Long complex chains can be created but we never say 'Modify a GE that is modding another GE'
 		OurCopyOfSpec->MakeUnique();
 	}
 
 	// Now that we have our own copy, apply our GEs that modify IncomingGEs
 	ActiveGameplayEffects.ApplyActiveEffectsTo(*OurCopyOfSpec, FModifierQualifier(BaseQualifier).Type(EGameplayMod::IncomingGE).IgnoreHandle(MyHandle));
+
+	// todo: apply some better logic to this so that we don't recalculate stacking effects as often
+	ActiveGameplayEffects.bNeedToRecalculateStacks = true;
 	
 	// Now that we have the final version of this effect, actually apply it if its going to be hanging around
 	if (Spec.GetDuration() != UGameplayEffect::INSTANT_APPLICATION )
@@ -296,10 +299,10 @@ FActiveGameplayEffectHandle UAttributeComponent::ApplyGameplayEffectSpecToSelf(c
 	if (bInvokeGameplayCueApplied)
 	{
 		// We both added and activated the GameplayCue here.
-		// On the client, who will invoke the gameplaycue from an OnRep, he will need to look at the StartTime to determine
+		// On the client, who will invoke the gameplay cue from an OnRep, he will need to look at the StartTime to determine
 		// if the Cue was actually added+activated or just added (due to relevancy)
 
-		// Fixme: what if we wanted to scale Cue magnitde based on damage? E.g, scale an cue effect when the GE is buffed?
+		// Fixme: what if we wanted to scale Cue magnitude based on damage? E.g, scale an cue effect when the GE is buffed?
 		InvokeGameplayCueAdded(*OurCopyOfSpec);
 		InvokeGameplayCueActivated(*OurCopyOfSpec);
 	}
@@ -307,7 +310,7 @@ FActiveGameplayEffectHandle UAttributeComponent::ApplyGameplayEffectSpecToSelf(c
 	// Execute the GE at least once (if instant, this will execute once and be done. If persistent, it was added to ActiveGameplayEffects above)
 	
 	// Execute if this is an instant application effect or if it is a periodic effect (this would be the first execution of a repeating effect)
-	if (Spec.GetDuration() == UGameplayEffect::INSTANT_APPLICATION || Spec.GetPeriod() != UGameplayEffect::NO_PERIOD)
+	if (Spec.GetDuration() == UGameplayEffect::INSTANT_APPLICATION)
 	{
 		ExecuteGameplayEffect(*OurCopyOfSpec, FModifierQualifier(BaseQualifier).IgnoreHandle(MyHandle));
 	}
@@ -328,8 +331,8 @@ void UAttributeComponent::ExecutePeriodicEffect(FActiveGameplayEffectHandle	Hand
 
 void UAttributeComponent::ExecuteGameplayEffect(const FGameplayEffectSpec &Spec, const FModifierQualifier &QualifierContext)
 {
-	// Should only ever execute effects that are instant applicatin or periodic application
-	// Effects with no period and that arent instant application should never be executed
+	// Should only ever execute effects that are instant application or periodic application
+	// Effects with no period and that aren't instant application should never be executed
 	check( (Spec.GetDuration() == UGameplayEffect::INSTANT_APPLICATION || Spec.GetPeriod() != UGameplayEffect::NO_PERIOD) );
 	
 	ActiveGameplayEffects.ExecuteActiveEffectsFrom(Spec, QualifierContext);
diff --git a/Engine/Source/Runtime/SkillSystem/Private/GameplayEffect.cpp b/Engine/Source/Runtime/SkillSystem/Private/GameplayEffect.cpp
--- a/Engine/Source/Runtime/SkillSystem/Private/GameplayEffect.cpp
+++ b/Engine/Source/Runtime/SkillSystem/Private/GameplayEffect.cpp
@@ -61,10 +61,90 @@ FGameplayEffectSpec::FGameplayEffectSpec(const UGameplayEffect * InDef, TSharedP
 	, ModifierLevel(InLevel)
 	, Duration(new FAggregator(InDef->Duration.MakeFinalizedCopy(CurveData), InLevel, SKILL_AGG_DEBUG(TEXT("%s Duration"), *InDef->GetName())))
 	, Period(new FAggregator(InDef->Period.MakeFinalizedCopy(CurveData), InLevel, SKILL_AGG_DEBUG(TEXT("%s Period"), *InDef->GetName())))
+	, StackingPolicy(EGameplayEffectStackingPolicy::Unlimited)
+	, StackedAttribName(NAME_None)
 {
 	Duration.Get()->RegisterLevelDependancies();
 	Period.Get()->RegisterLevelDependancies();
 
+	// check the stacking policy for this gameplay effect and warn if it's set up badly
+	if (Def->StackingPolicy != EGameplayEffectStackingPolicy::Unlimited)
+	{
+		StackingPolicy = Def->StackingPolicy;
+
+		if (Def->Modifiers.Num() < 1)
+		{
+			SKILL_LOG(Warning, TEXT("%s has a stacking rule but does not have any modifiers to apply it to."), *Def->GetPathName());
+			StackingPolicy = EGameplayEffectStackingPolicy::Unlimited;
+		}
+		else
+		{
+			UDataTable* DataTable = ISkillSystemModule::Get().GetSkillSystemGlobals().GetGlobalAttributeDataTable();
+
+			if (Def->Modifiers[0].ModifierType != EGameplayMod::Attribute)
+			{
+				SKILL_LOG(Warning, TEXT("%s has a stacking rule but its first modifier does not modify an attribute."), *Def->GetPathName());
+				StackingPolicy = EGameplayEffectStackingPolicy::Unlimited;
+			}
+			else
+			{
+				FName AttributeName(*Def->Modifiers[0].Attribute.GetName());
+				
+				if (DataTable && DataTable->RowMap.Contains(AttributeName))
+				{
+					FAttributeMetaData* Row = (FAttributeMetaData*)DataTable->RowMap[AttributeName];
+					if (Row && Row->bCanStack == false)
+					{
+						SKILL_LOG(Warning, TEXT("%s has a stacking rule but modifies attribute %s. %s is not allowed to stack."), *Def->GetPathName(), *Def->Modifiers[0].Attribute.GetName(), *Def->Modifiers[0].Attribute.GetName());
+						StackingPolicy = EGameplayEffectStackingPolicy::Unlimited;
+					}
+				}
+				else
+				{
+					SKILL_LOG(Warning, TEXT("%s has a stacking rule but modifies attribute %s. %s was not found in the global attribute data table."), *Def->GetPathName(), *Def->Modifiers[0].Attribute.GetName(), *Def->Modifiers[0].Attribute.GetName());
+					StackingPolicy = EGameplayEffectStackingPolicy::Unlimited;
+				}
+			}
+
+			// look for other stacking attributes and warn that they will be ignored
+			for (int32 Idx = 1; Idx < Def->Modifiers.Num(); ++Idx)
+			{
+				if (Def->Modifiers[Idx].ModifierType == EGameplayMod::Attribute)
+				{
+					FName AttributeName(*Def->Modifiers[0].Attribute.GetName());
+
+					if (DataTable && DataTable->RowMap.Contains(AttributeName))
+					{
+						FAttributeMetaData* Row = (FAttributeMetaData*)DataTable->RowMap[AttributeName];
+						if (Row && Row->bCanStack)
+						{
+							SKILL_LOG(Warning, TEXT("%s has a stacking rule and modifies attribute %s but %s is not the modified by the first modifier. Stacking is based on the first modifier."), *Def->GetPathName(), *Def->Modifiers[0].Attribute.GetName(), *Def->Modifiers[0].Attribute.GetName());
+						}
+					}
+				}
+			}
+		}
+		
+		UDataTable* DataTable = ISkillSystemModule::Get().GetSkillSystemGlobals().GetGlobalAttributeDataTable();
+
+		// check for other attributes that may allow stacking and warn that they will be ignored
+		for (int32 Idx = 1; Idx < Def->Modifiers.Num(); ++Idx)
+		{
+			if (Def->Modifiers[Idx].ModifierType == EGameplayMod::Attribute)
+			{
+				FName AttributeName(*Def->Modifiers[Idx].Attribute.GetName());
+				if (DataTable && DataTable->RowMap.Contains(AttributeName))
+				{
+					FAttributeMetaData* Row = (FAttributeMetaData*)DataTable->RowMap[AttributeName];
+					if (Row && Row->bCanStack == true)
+					{
+						SKILL_LOG(Warning, TEXT("%s has a stacking rule and modifies more than one stackable attribute. Stacking will  not apply to attribute %s"), *Def->GetPathName(), *Def->Modifiers[Idx].Attribute.GetName());
+					}
+				}
+			}
+		}
+	}
+
 	InitModifiers(CurveData);
 }
 
@@ -180,6 +260,36 @@ float FGameplayEffectSpec::GetPeriod() const
 	return Period.Get()->Evaluate().Magnitude;
 }
 
+float FGameplayEffectSpec::GetMagnitude(const FGameplayAttribute &Attribute) const
+{
+	float CurrentMagnitude = 0.f;
+
+	for (const FModifierSpec &Mod : Modifiers)
+	{
+		if (Mod.Info.ModifierType == EGameplayMod::Attribute)
+		{
+			if (Mod.Info.Attribute != Attribute)
+			{
+				continue;
+			}
+
+			// Todo: Tags/application checks here - make sure we can still apply
+
+			// First, evaluate all of our data 
+
+			FGameplayModifierEvaluatedData EvaluatedData = Mod.Aggregator.Get()->Evaluate();
+
+			FAggregator	Aggregator(CurrentMagnitude, SKILL_AGG_DEBUG(TEXT("Magnitude of Attribute %s"), *Mod.Info.Attribute.GetName()));
+
+			Aggregator.ExecuteMod(Mod.Info.ModifierOp, EvaluatedData);
+
+			CurrentMagnitude = Aggregator.Evaluate().Magnitude;
+		}
+	}
+
+	return CurrentMagnitude;
+}
+
 bool FGameplayEffectSpec::ShouldApplyAsSnapshot(const FModifierQualifier &QualifierContext) const
 {
 	bool ShouldSnapshot;
@@ -201,6 +311,11 @@ bool FGameplayEffectSpec::ShouldApplyAsSnapshot(const FModifierQualifier &Qualif
 	return ShouldSnapshot;
 }
 
+EGameplayEffectStackingPolicy::Type FGameplayEffectSpec::GetStackingType() const
+{
+	return StackingPolicy;
+}
+
 // --------------------------------------------------------------------------------------------------------------------------------------------------------
 //
 //	FModifierSpec
@@ -850,7 +965,7 @@ void FGameplayModifierEvaluatedData::InvokePostExecute(const FGameplayEffectModC
 
 // --------------------------------------------------------------------------------------------------------------------------------------------------------
 //
-//	FGameplayModifierEvaluatedData
+//	FActiveGameplayEffect
 //
 // --------------------------------------------------------------------------------------------------------------------------------------------------------
 
@@ -924,7 +1039,7 @@ void FActiveGameplayEffectsContainer::ApplySpecToActiveEffectsAndAttributes(cons
 
 			// TODO: Tag checks here
 
-			// This modies GEs that are currently active, so apply this to them...
+			// This modifies GEs that are currently active, so apply this to them...
 			for (FActiveGameplayEffect & ActiveEffect : GameplayEffects)
 			{
 				if (!QualifierContext.TestTarget(ActiveEffect.Handle))
@@ -943,6 +1058,17 @@ void FActiveGameplayEffectsContainer::ExecuteActiveEffectsFrom(const FGameplayEf
 {
 	bool InvokeGameplayCueExecute = false;
 
+	// check if this is a stacking effect and if it is the active stacking effect.
+	// todo: move this into the loop if we end up applying stacking to different modifiers on the same effect
+	if (Spec.GetStackingType() != EGameplayEffectStackingPolicy::Unlimited)
+	{
+		// we're a stacking attribute but not active
+		if (Spec.StackedAttribName == NAME_None)
+		{
+			return;
+		}
+	}
+
 	for (const FModifierSpec &Mod : Spec.Modifiers)
 	{
 		if (Mod.Info.ModifierType == EGameplayMod::Attribute)
@@ -1006,7 +1132,7 @@ void FActiveGameplayEffectsContainer::ExecuteActiveEffectsFrom(const FGameplayEf
 
 			// Todo: Tag checks here
 
-			// This modies GEs that are currently active, so apply this to them...
+			// This modifies GEs that are currently active, so apply this to them...
 			for (FActiveGameplayEffect & ActiveEffect : GameplayEffects)
 			{
 				// Don't apply spec to itself
@@ -1025,7 +1151,7 @@ void FActiveGameplayEffectsContainer::ExecuteActiveEffectsFrom(const FGameplayEf
 
 	if (InvokeGameplayCueExecute)
 	{
-		// TODO: check replication policy. Right now we will replciate every execute via a multicast RPC
+		// TODO: check replication policy. Right now we will replicate every execute via a multicast RPC
 
 		SKILL_LOG(Log, TEXT("Invoking Execute GameplayCue for %s"), *Spec.ToSimpleString() );
 		Owner->NetMulticast_InvokeGameplayCueExecuted(Spec);
@@ -1040,7 +1166,6 @@ bool FActiveGameplayEffectsContainer::ExecuteGameplayEffect(FActiveGameplayEffec
 	{
 		if (GameplayEffects[idx].Handle == Handle)
 		{
-			//GameplayEffects.
 			return true;
 		}
 	}
@@ -1112,13 +1237,20 @@ float FActiveGameplayEffectsContainer::GetGameplayEffectMagnitude(FActiveGamepla
 bool FActiveGameplayEffectsContainer::IsGameplayEffectActive(FActiveGameplayEffectHandle Handle) const
 {
 	// Could make this a map for quicker lookup
-	for (FActiveGameplayEffect Effect : GameplayEffects)
+	for (const FActiveGameplayEffect& Effect : GameplayEffects)
 	{
 		if (Effect.Handle == Handle)
 		{
+			// stacking effects may return false
+			if (Effect.Spec.GetStackingType() != EGameplayEffectStackingPolicy::Unlimited &&
+				Effect.Spec.StackedAttribName == NAME_None)
+			{
+				return false;
+			}
 			return true;
 		}
 	}
+
 	return false;
 }
 
@@ -1220,20 +1352,29 @@ void FActiveGameplayEffectsContainer::TEMP_TickActiveEffects(float DeltaSeconds)
 		CurrentTime = Owner->GetWorld()->GetGameState<AGameState>()->ElapsedTime;
 	}
 
-	for (int32 idx = 0; idx < GameplayEffects.Num(); ++idx)
+	// effects may have been added outside of the tick during the last frame
+	if (bNeedToRecalculateStacks)
+	{
+		RecalculateStacking();
+	}
+
+	for (int32 Idx = 0; Idx < GameplayEffects.Num(); ++Idx)
 	{
-		FActiveGameplayEffect &Effect = GameplayEffects[idx];
+		FActiveGameplayEffect &Effect = GameplayEffects[Idx];
 
 		if (Effect.NextExecuteTime > 0)
 		{
 			float TimeOver = CurrentTime - Effect.NextExecuteTime;
-			if (TimeOver >= -KINDA_SMALL_NUMBER)
+			while (TimeOver >= -KINDA_SMALL_NUMBER)
 			{
 				// Advance
 				Effect.AdvanceNextExecuteTime(CurrentTime);
 
 				// Execute
 				ExecuteActiveEffectsFrom( Effect.Spec, FModifierQualifier().IgnoreHandle(Effect.Handle) );
+
+				// Update TimeOver
+				TimeOver = CurrentTime - Effect.NextExecuteTime;
 			}
 		}
 
@@ -1243,10 +1384,126 @@ void FActiveGameplayEffectsContainer::TEMP_TickActiveEffects(float DeltaSeconds)
 			Owner->InvokeGameplayCueRemoved(Effect.Spec);
 
 			MarkArrayDirty();
-			GameplayEffects.RemoveAtSwap(idx);
-			idx--;
+			GameplayEffects.RemoveAtSwap(Idx);
+			Idx--;
+
+			// todo: check if the removed effect deals with an attribute that stacks, if not don't set the flag
+			bNeedToRecalculateStacks = true;
 		}
 	}
+
+	// this needs to be updated here instead of waiting for the start of the next tick
+	// an instantaneous effect could come in before the start of the next tick
+	if (bNeedToRecalculateStacks)
+	{
+		RecalculateStacking();
+	}
+}
+
+void FActiveGameplayEffectsContainer::RecalculateStacking()
+{
+	bNeedToRecalculateStacks = false;
+
+	// one or more gameplay effects has been removed or added so we need to go through all of them to see what the best elements are in each stack
+	TArray<FActiveGameplayEffect*> StackedEffects;
+	TArray<FActiveGameplayEffect*> CustomStackedEffects;
+
+	for (FActiveGameplayEffect& Effect : GameplayEffects)
+	{
+		// ignore effects that don't stack and effects that replace when they stack
+		// effects that replace when they stack only need to be handled when they are added
+		if (Effect.Spec.GetStackingType() != EGameplayEffectStackingPolicy::Unlimited &&
+			Effect.Spec.GetStackingType() != EGameplayEffectStackingPolicy::Replaces)
+		{
+			if (Effect.Spec.Def->Modifiers.Num() > 0)
+			{
+				Effect.Spec.StackedAttribName = NAME_None;
+				bool bFoundStack = false;
+
+				// group all of the custom stacking effects and deal with them after
+				if (Effect.Spec.GetStackingType() == EGameplayEffectStackingPolicy::Callback)
+				{
+					CustomStackedEffects.Add(&Effect);
+					continue;
+				}
+
+				for (FActiveGameplayEffect*& StackedEffect : StackedEffects)
+				{
+					if (StackedEffect->Spec.GetStackingType() == Effect.Spec.GetStackingType() && StackedEffect->Spec.Def->Modifiers[0].Attribute == Effect.Spec.Def->Modifiers[0].Attribute)
+					{
+						switch (Effect.Spec.GetStackingType())
+						{
+						case EGameplayEffectStackingPolicy::Highest:
+						{
+							float BestSpecMagnitude = StackedEffect->Spec.GetMagnitude(StackedEffect->Spec.Modifiers[0].Info.Attribute);
+							float CurrSpecMagnitude = Effect.Spec.GetMagnitude(Effect.Spec.Modifiers[0].Info.Attribute);
+							if (BestSpecMagnitude < CurrSpecMagnitude)
+							{
+								StackedEffect = &Effect;
+							}
+							break;
+						}
+						case EGameplayEffectStackingPolicy::Lowest:
+						{
+							float BestSpecMagnitude = StackedEffect->Spec.GetMagnitude(StackedEffect->Spec.Modifiers[0].Info.Attribute);
+							float CurrSpecMagnitude = Effect.Spec.GetMagnitude(Effect.Spec.Modifiers[0].Info.Attribute);
+							if (BestSpecMagnitude > CurrSpecMagnitude)
+							{
+								StackedEffect = &Effect;
+							}
+							break;
+						}
+						default:
+							// we shouldn't ever be here
+							SKILL_LOG(Warning, TEXT("%s uses unhandled stacking rule %s."), *Effect.Spec.Def->GetPathName(), *EGameplayEffectStackingPolicyToString(Effect.Spec.GetStackingType()));
+							break;
+						};
+
+						bFoundStack = true;
+						break;
+					}
+				}
+
+				// if we didn't find a stack matching this effect we need to add it to our list
+				if (!bFoundStack)
+				{
+					StackedEffects.Add(&Effect);
+				}
+			}
+		}
+	}
+
+	// now deal with the custom effects
+	while (CustomStackedEffects.Num() > 0)
+	{
+		int32 StackedIdx = StackedEffects.Add(CustomStackedEffects[0]);
+		CustomStackedEffects.RemoveAtSwap(0);
+
+		UGameplayEffectStackingExtension * StackedExt = StackedEffects[StackedIdx]->Spec.Def->StackingExtension->GetDefaultObject<UGameplayEffectStackingExtension>();
+
+		TArray<FActiveGameplayEffect*> Effects;
+		
+		for (int32 Idx = 0; Idx < CustomStackedEffects.Num(); Idx++)
+		{
+			UGameplayEffectStackingExtension * CurrentExt = CustomStackedEffects[Idx]->Spec.Def->StackingExtension->GetDefaultObject<UGameplayEffectStackingExtension>();
+			if (CurrentExt && StackedExt && CurrentExt->Handle == StackedExt->Handle && 
+				StackedEffects[StackedIdx]->Spec.Def->Modifiers[0].Attribute == CustomStackedEffects[Idx]->Spec.Def->Modifiers[0].Attribute)
+			{
+				Effects.Add(CustomStackedEffects[Idx]);
+				CustomStackedEffects.RemoveAtSwap(Idx);
+				--Idx;
+			}
+		}
+
+		UGameplayEffectStackingExtension * Ext = StackedEffects[StackedIdx]->Spec.Def->StackingExtension->GetDefaultObject<UGameplayEffectStackingExtension>();
+		Ext->CalculateStack(Effects, *this, *StackedEffects[StackedIdx]);
+	}
+
+	// we've found all of the best stacking effects, mark them so that they updated properly
+	for (FActiveGameplayEffect* const StackedEffect : StackedEffects)
+	{
+		StackedEffect->Spec.StackedAttribName = FName(*StackedEffect->Spec.Def->Modifiers[0].Attribute.GetName());
+	}
 }
 
 // --------------------------------------------------------------------------------------------------------------------------------------------------------
@@ -1328,3 +1585,9 @@ FString EGameplayEffectCopyPolicyToString(int32 Type)
 	return e->GetEnum(Type).ToString();
 }
 
+FString EGameplayEffectStackingPolicyToString(int32 Type)
+{
+	static UEnum *e = FindObject<UEnum>(ANY_PACKAGE, TEXT("EGameplayEffectStackingPolicy"));
+	return e->GetEnum(Type).ToString();
+}
+
diff --git a/Engine/Source/Runtime/SkillSystem/Private/GameplayEffectStackingExtension.cpp b/Engine/Source/Runtime/SkillSystem/Private/GameplayEffectStackingExtension.cpp
new file mode 100644
--- /dev/null
+++ b/Engine/Source/Runtime/SkillSystem/Private/GameplayEffectStackingExtension.cpp
@@ -0,0 +1,14 @@
+// Copyright 1998-2014 Epic Games, Inc. All Rights Reserved.
+
+#include "SkillSystemModulePrivatePCH.h"
+
+UGameplayEffectStackingExtension::UGameplayEffectStackingExtension(const class FPostConstructInitializeProperties& PCIP)
+: Super(PCIP)
+{
+	
+}
+
+FGameplayEffectSpec* UGameplayEffectStackingExtension::StackCustomEffects(TArray<FGameplayEffectSpec*>& CustomGameplayEffects)
+{
+	return NULL;
+}
\ No newline at end of file
diff --git a/Engine/Source/Runtime/SkillSystem/Private/GameplayEffectStackingExtension_CappedNumberTest.cpp b/Engine/Source/Runtime/SkillSystem/Private/GameplayEffectStackingExtension_CappedNumberTest.cpp
new file mode 100644
--- /dev/null
+++ b/Engine/Source/Runtime/SkillSystem/Private/GameplayEffectStackingExtension_CappedNumberTest.cpp
@@ -0,0 +1,73 @@
+// Copyright 1998-2014 Epic Games, Inc. All Rights Reserved.
+
+#include "SkillSystemModulePrivatePCH.h"
+#include "Crc.h"
+
+const int32 MaxNumInStack = 2;
+
+UGameplayEffectStackingExtension_CappedNumberTest::UGameplayEffectStackingExtension_CappedNumberTest(const class FPostConstructInitializeProperties& PCIP)
+: Super(PCIP)
+{
+	Handle = FCrc::StrCrc32("UGameplayEffectStackingExtension_CappedNumberTest");
+}
+
+void UGameplayEffectStackingExtension_CappedNumberTest::CalculateStack(TArray<FActiveGameplayEffect*>& CustomGameplayEffects, FActiveGameplayEffectsContainer& Container, FActiveGameplayEffect& CurrentEffect)
+{
+	// this effect shouldn't be in the array so be sure to count it as well
+	int32 EffectiveCount = FMath::Min(CustomGameplayEffects.Num() + 1, MaxNumInStack);
+	
+	// find the most recent time one of these was applied
+	float StartTime = CurrentEffect.StartTime;
+
+	for (FActiveGameplayEffect* Effect : CustomGameplayEffects)
+	{
+		StartTime = FMath::Max(StartTime, Effect->StartTime);
+	}
+
+	{
+		// set the start time to be equal to the most recent start time so that the stacked effects stick around if a new effect has been added
+		CurrentEffect.StartTime = StartTime;
+		int32 Idx = 0;
+		while (Idx < FMath::Min(MaxNumInStack - 1, CustomGameplayEffects.Num()))
+		{
+			CustomGameplayEffects[Idx]->StartTime = StartTime;
+			++Idx;
+		}
+
+		// we don't need any effects beyond the cap
+		while (Idx < CustomGameplayEffects.Num())
+		{
+			Container.RemoveActiveGameplayEffect(CustomGameplayEffects[Idx]->Handle);
+ 			++Idx;
+		}
+	}
+
+	for (FModifierSpec Mod : CurrentEffect.Spec.Modifiers)
+	{
+		if (Mod.Info.OwnedTags.HasTag("Stackable"))
+		{
+			// remove any stacking information that was already applied to the current modifier
+			for (int32 Idx = 0; Idx < Mod.Aggregator.Get()->Mods[EGameplayModOp::Multiplicitive].Num(); ++Idx)
+			{
+				FAggregatorRef& Agg = Mod.Aggregator.Get()->Mods[EGameplayModOp::Multiplicitive][Idx];
+				if (Agg.Get()->BaseData.Tags.HasTag("Stack.CappedNumber"))
+				{
+					Mod.Aggregator.Get()->Mods[EGameplayModOp::Multiplicitive].RemoveAtSwap(Idx);
+					--Idx;
+				}
+			}
+			FGameplayModifierInfo ModInfo;
+			ModInfo.Magnitude.SetValue(EffectiveCount);
+			ModInfo.ModifierOp = EGameplayModOp::Multiplicitive;
+			ModInfo.OwnedTags.AddTag("Stack.CappedNumber");
+			ModInfo.Attribute = Mod.Info.Attribute;
+
+			TSharedPtr<FGameplayEffectLevelSpec> ModifierLevel(TSharedPtr< FGameplayEffectLevelSpec >(new FGameplayEffectLevelSpec()));
+			ModifierLevel->ApplyNewDef(ModInfo.LevelInfo, ModifierLevel);
+
+			FModifierSpec ModSpec(ModInfo, ModifierLevel, NULL);
+
+			ModSpec.ApplyModTo(Mod, true);
+		}
+	}
+}
diff --git a/Engine/Source/Runtime/SkillSystem/Private/GameplayEffectStackingExtension_DiminishingReturnsTest.cpp b/Engine/Source/Runtime/SkillSystem/Private/GameplayEffectStackingExtension_DiminishingReturnsTest.cpp
new file mode 100644
--- /dev/null
+++ b/Engine/Source/Runtime/SkillSystem/Private/GameplayEffectStackingExtension_DiminishingReturnsTest.cpp
@@ -0,0 +1,93 @@
+// Copyright 1998-2014 Epic Games, Inc. All Rights Reserved.
+
+#include "SkillSystemModulePrivatePCH.h"
+
+UCurveTable* GetCurveTable()
+{
+	FString CSV(TEXT(", 0, 1, 2, 3, 100\r\nReturnValue, 0, 5, 7, 8, 105"));
+
+	UCurveTable* CurveTable = Cast<UCurveTable>(StaticConstructObject(UCurveTable::StaticClass(), GetTransientPackage(), FName(TEXT("TempCurveTable"))));
+	CurveTable->CreateTableFromCSVString(CSV);
+
+	return CurveTable;
+}
+
+UGameplayEffectStackingExtension_DiminishingReturnsTest::UGameplayEffectStackingExtension_DiminishingReturnsTest(const class FPostConstructInitializeProperties& PCIP)
+: Super(PCIP)
+{
+	Handle = FCrc::StrCrc32("UGameplayEffectStackingExtension_DiminishingReturnsTest");
+}
+
+void UGameplayEffectStackingExtension_DiminishingReturnsTest::CalculateStack(TArray<FActiveGameplayEffect*>& CustomGameplayEffects, FActiveGameplayEffectsContainer& Container, FActiveGameplayEffect& CurrentEffect)
+{
+	// look up the correct value in the curve table
+	UCurveTable* CurveTable = GetCurveTable();
+	float Multiplier = 1.f;
+	if (CurveTable)
+	{
+		FRichCurve* RichCurve = CurveTable->FindCurve(FName(TEXT("ReturnValue")), TEXT("Stack"));
+		if (RichCurve)
+		{
+			Multiplier = RichCurve->Eval(CustomGameplayEffects.Num() + 1); // + 1 because the current effect isn't in the TArray
+		}
+	}
+
+// 	// remove any old stack modifiers
+// 	for (int32 Idx = 0; Idx < CurrentEffect.Spec.Modifiers.Num(); ++Idx)
+// 	{
+// 		if (CurrentEffect.Spec.Modifiers[Idx].Info.OwnedTags.HasTag("Stack"))
+// 		{
+// 			CurrentEffect.Spec.Modifiers.RemoveAtSwap(Idx);
+// 			--Idx;
+// 		}
+// 	}
+// 	
+// 	for (FModifierSpec Mod : CurrentEffect.Spec.Modifiers)
+// 	{
+// 		if (Mod.Info.OwnedTags.HasTag("Stackable"))
+// 		{
+// 			FGameplayModifierInfo ModInfo;
+// 			ModInfo.Magnitude.SetValue(Multiplier);
+// 			ModInfo.ModifierOp = EGameplayModOp::Multiplicitive;
+// 			ModInfo.OwnedTags.AddTag("Stack");
+// 			ModInfo.Attribute = Mod.Info.Attribute;
+// 
+// 			TSharedPtr<FGameplayEffectLevelSpec> ModifierLevel(TSharedPtr< FGameplayEffectLevelSpec >(new FGameplayEffectLevelSpec()));
+// 			ModifierLevel->ApplyNewDef(ModInfo.LevelInfo, ModifierLevel);
+// 
+// 			FModifierSpec ModSpec(ModInfo, ModifierLevel, NULL);
+// 
+// 			ModSpec.ApplyModTo(Mod, false);
+// 			CurrentEffect.Spec.Modifiers.Add(ModSpec);
+// 		}
+// 	}
+
+	for (FModifierSpec Mod : CurrentEffect.Spec.Modifiers)
+	{
+		if (Mod.Info.OwnedTags.HasTag("Stackable"))
+		{
+			// remove any stacking information that was already applied to the current modifier
+			for (int32 Idx = 0; Idx < Mod.Aggregator.Get()->Mods[EGameplayModOp::Multiplicitive].Num(); ++Idx)
+			{
+				FAggregatorRef& Agg = Mod.Aggregator.Get()->Mods[EGameplayModOp::Multiplicitive][Idx];
+				if (Agg.Get()->BaseData.Tags.HasTag("Stack.DiminishingReturns"))
+				{
+					Mod.Aggregator.Get()->Mods[EGameplayModOp::Multiplicitive].RemoveAtSwap(Idx);
+					--Idx;
+				}
+			}
+			FGameplayModifierInfo ModInfo;
+			ModInfo.Magnitude.SetValue(Multiplier);
+			ModInfo.ModifierOp = EGameplayModOp::Multiplicitive;
+			ModInfo.OwnedTags.AddTag("Stack.DiminishingReturns");
+			ModInfo.Attribute = Mod.Info.Attribute;
+
+			TSharedPtr<FGameplayEffectLevelSpec> ModifierLevel(TSharedPtr< FGameplayEffectLevelSpec >(new FGameplayEffectLevelSpec()));
+			ModifierLevel->ApplyNewDef(ModInfo.LevelInfo, ModifierLevel);
+
+			FModifierSpec ModSpec(ModInfo, ModifierLevel, NULL);
+
+			ModSpec.ApplyModTo(Mod, true);
+		}
+	}
+}
diff --git a/Engine/Source/Runtime/SkillSystem/Private/GameplayEffectTests.cpp b/Engine/Source/Runtime/SkillSystem/Private/GameplayEffectTests.cpp
--- a/Engine/Source/Runtime/SkillSystem/Private/GameplayEffectTests.cpp
+++ b/Engine/Source/Runtime/SkillSystem/Private/GameplayEffectTests.cpp
@@ -124,6 +124,9 @@ bool GameplayEffectsTest_InstantDamage(UWorld *World, FAutomationTestBase * Test
 		SKILL_LOG(Log, TEXT("Final Health: %.2f"), DestComponent->GetSet<USkillSystemTestAttributeSet>()->Health);
 	}
 
+	World->EditorDestroyActor(SourceActor, false);
+	World->EditorDestroyActor(DestActor, false);
+
 	return true;
 }
 
@@ -174,6 +177,9 @@ bool GameplayEffectsTest_InstantDamageRemap(UWorld *World, FAutomationTestBase *
 		}
 	}
 
+	World->EditorDestroyActor(SourceActor, false);
+	World->EditorDestroyActor(DestActor, false);
+
 	return true;
 }
 
@@ -232,6 +238,9 @@ bool GameplayEffectsTest_InstantDamage_Buffed(UWorld *World, FAutomationTestBase
 		SKILL_LOG(Log, TEXT("Final Health: %.2f"), DestComponent->GetSet<USkillSystemTestAttributeSet>()->Health);
 	}
 
+	World->EditorDestroyActor(SourceActor, false);
+	World->EditorDestroyActor(DestActor, false);
+
 	return true;
 }
 
@@ -287,6 +296,9 @@ bool GameplayEffectsTest_TemporaryDamage(UWorld *World, FAutomationTestBase * Te
 		SKILL_LOG(Log, TEXT("After Removal. Health: %.2f. RemovedEffecte: %d"), DestComponent->GetSet<USkillSystemTestAttributeSet>()->Health, RemovedEffect);
 	}
 
+	World->EditorDestroyActor(SourceActor, false);
+	World->EditorDestroyActor(DestActor, false);
+
 	return true;
 }
 
@@ -374,6 +386,9 @@ bool GameplayEffectsTest_TemporaryDamageBuffed(UWorld *World, FAutomationTestBas
 		Test->TestTrue(SKILL_TEST_TEXT("Number of GameplayEffects=0"), DestComponent->GetNumActiveGameplayEffect() == 0);
 	}
 
+	World->EditorDestroyActor(SourceActor, false);
+	World->EditorDestroyActor(DestActor, false);
+
 	return true;
 }
 
@@ -480,6 +495,9 @@ bool GameplayEffectsTest_TemporaryDamageTemporaryBuff(UWorld *World, FAutomation
 		Test->TestTrue(SKILL_TEST_TEXT("GameplayEffectsTest_TemporaryDamageTemporaryBuff Number of GameplayEffects=0"), DestComponent->GetNumActiveGameplayEffect() == 0);
 	}
 
+	World->EditorDestroyActor(SourceActor, false);
+	World->EditorDestroyActor(DestActor, false);
+
 	return true;
 }
 
@@ -564,6 +582,9 @@ bool GameplayEffectsTest_LinkedBuffDestroy(UWorld *World, FAutomationTestBase *
 
 	// DestComponent->PrintAllGameplayEffects();
 
+	World->EditorDestroyActor(SourceActor, false);
+	World->EditorDestroyActor(DestActor, false);
+
 	return true;
 }
 
@@ -650,6 +671,9 @@ bool GameplayEffectsTest_SnapshotBuffDestroy(UWorld *World, FAutomationTestBase
 
 	// DestComponent->PrintAllGameplayEffects();
 
+	World->EditorDestroyActor(SourceActor, false);
+	World->EditorDestroyActor(DestActor, false);
+
 	return true;
 }
 
@@ -736,6 +760,10 @@ bool GameplayEffectsTest_DurationBuff(UWorld *World, FAutomationTestBase * Test)
 	}
 
 	//DestComponent->PrintAllGameplayEffects();
+
+	World->EditorDestroyActor(SourceActor, false);
+	World->EditorDestroyActor(DestActor, false);
+
 	return true;
 }
 
@@ -872,6 +900,9 @@ bool GameplayEffectsTest_DamageBuffBuff_Basic(UWorld *World, FAutomationTestBase
 
 	// DestComponent->PrintAllGameplayEffects();
 
+	World->EditorDestroyActor(SourceActor, false);
+	World->EditorDestroyActor(DestActor, false);
+
 	return true;
 }
 
@@ -1009,6 +1040,9 @@ bool GameplayEffectsTest_DamageBuffBuff_FullLink(UWorld *World, FAutomationTestB
 
 	// DestComponent->PrintAllGameplayEffects();
 
+	World->EditorDestroyActor(SourceActor, false);
+	World->EditorDestroyActor(DestActor, false);
+
 	return true;
 }
 
@@ -1144,6 +1178,9 @@ bool GameplayEffectsTest_DamageBuffBuff_FullSnapshot(UWorld *World, FAutomationT
 
 	// DestComponent->PrintAllGameplayEffects();
 
+	World->EditorDestroyActor(SourceActor, false);
+	World->EditorDestroyActor(DestActor, false);
+
 	return true;
 }
 
@@ -1279,6 +1316,9 @@ bool GameplayEffectsTest_DamageBuffBuff_SnapshotLink(UWorld *World, FAutomationT
 
 	// DestComponent->PrintAllGameplayEffects();
 
+	World->EditorDestroyActor(SourceActor, false);
+	World->EditorDestroyActor(DestActor, false);
+
 	return true;
 }
 
@@ -1308,6 +1348,9 @@ bool TimerTest(UWorld *World, FAutomationTestBase * Test)
 	World->Tick(ELevelTick::LEVELTICK_All, 3.f);
 	World->Tick(ELevelTick::LEVELTICK_All, 3.f);
 
+	World->EditorDestroyActor(SourceActor, false);
+	World->EditorDestroyActor(DestActor, false);
+
 	return true;
 }
 
@@ -1381,6 +1424,9 @@ bool GameplayEffectsTest_DurationDamage(UWorld *World, FAutomationTestBase * Tes
 		Test->TestTrue(SKILL_TEST_TEXT("NumberOfActive Dest GameplayEffects %d == 0", NumEffects), NumEffects == 0);
 	}
 
+	World->EditorDestroyActor(SourceActor, false);
+	World->EditorDestroyActor(DestActor, false);
+
 	return true;
 }
 
@@ -1420,15 +1466,15 @@ bool GameplayEffectsTest_PeriodicDamage(UWorld *World, FAutomationTestBase * Tes
 		BaseDmgEffect->GameplayCues.Add( FGameplayEffectCue(FName(TEXT("GameplayCue.Burning")), 1.f, 10.f) );
 
 		SourceComponent->ApplyGameplayEffectSpecToTarget(BaseDmgEffect, DestComponent, 5.f);
-		ApplyCount++;
 
-		// The effect should instantly execute one time without ticking (for now at least)
+		// The effect should execute on the next tick
 		float ActualValue = DestComponent->GetSet<USkillSystemTestAttributeSet>()->Health;
 		float ExpectedValue = StartHealth + (DamageValue * ApplyCount);
 
 		Test->TestTrue(SKILL_TEST_TEXT("Damage Applied. Actual: %.2f == Exected: %.2f", ActualValue, ExpectedValue), ActualValue == ExpectedValue );
 	}
 		
+	ApplyCount++; // the effect will execute as soon as we tick any amount of time
 	for (int32 i=0; i < 10; ++i)
 	{
 		GameplayTest_TickWorld(World, 1.f);
@@ -1451,6 +1497,9 @@ bool GameplayEffectsTest_PeriodicDamage(UWorld *World, FAutomationTestBase * Tes
 		Test->TestTrue(SKILL_TEST_TEXT("NumberOfActive Dest GameplayEffects %d == 0", NumEffects), NumEffects == 0);
 	}
 
+	World->EditorDestroyActor(SourceActor, false);
+	World->EditorDestroyActor(DestActor, false);
+
 	return true;
 }
 
@@ -1525,10 +1574,13 @@ bool GameplayEffectsTest_LifestealExtension(UWorld *World, FAutomationTestBase *
 
 	}
 
+	World->EditorDestroyActor(SourceActor, false);
+	World->EditorDestroyActor(DestActor, false);
+
 	return true;
 }
 
-bool GameplayEffectsTest_ShieldlExtension(UWorld *World, FAutomationTestBase * Test)
+bool GameplayEffectsTest_ShieldExtension(UWorld *World, FAutomationTestBase * Test)
 {
 	const float StartHealth = 100.f;
 	const float DamageValue = -10.f;
@@ -1636,10 +1688,13 @@ bool GameplayEffectsTest_ShieldlExtension(UWorld *World, FAutomationTestBase * T
 		}
 	}
 
+	World->EditorDestroyActor(SourceActor, false);
+	World->EditorDestroyActor(DestActor, false);
+
 	return true;
 }
 
-bool GameplayEffectsTest_ShieldlExtensionMultiple(UWorld *World, FAutomationTestBase * Test)
+bool GameplayEffectsTest_ShieldExtensionMultiple(UWorld *World, FAutomationTestBase * Test)
 {
 	// This applies 2 instances of the shield and confirms that only 1 will absorb damage at at time
 
@@ -1780,6 +1835,9 @@ bool GameplayEffectsTest_ShieldlExtensionMultiple(UWorld *World, FAutomationTest
 		}
 	}
 
+	World->EditorDestroyActor(SourceActor, false);
+	World->EditorDestroyActor(DestActor, false);
+
 	return true;
 }
 
@@ -1825,6 +1883,35 @@ UCurveTable * GetStandardDamageOverrideCurveTable(float Factor)
 	return CurveTable;
 }
 
+UDataTable * SetGlobalDataTable()
+{
+	// set up a test table where SpellDamage stacks and PhysicalDamage does not.
+	FString CSV(TEXT(",BaseValue,MinValue,MaxValue,DerivedAttributeInfo,bCanStack\r\nStackingAttribute1,0.0,-999.9,999.9,,True\r\nStackingAttribute2,0.0,-999.9,999.9,,True\r\nNoStackAttribute,0.0,-999.9,999.9,,False\r\n"));
+
+	UDataTable * DataTable = Cast<UDataTable>(StaticConstructObject(UDataTable::StaticClass(), GetTransientPackage(), FName(TEXT("TempDataTable"))));
+	DataTable->RowStruct = FAttributeMetaData::StaticStruct();
+	DataTable->CreateTableFromCSVString(CSV);
+
+	FAttributeMetaData * Row = (FAttributeMetaData*)DataTable->RowMap["StackingAttribute1"];
+	if (Row)
+	{
+		check(Row->bCanStack);
+	}
+	Row = (FAttributeMetaData*)DataTable->RowMap["NoStackAttribute"];
+	if (Row)
+	{
+		check(!Row->bCanStack);
+	}
+
+	ISkillSystemModule::Get().GetSkillSystemGlobals().AutomationTestOnly_SetGlobalAttributeDataTable(DataTable);
+	return DataTable;
+}
+
+void ClearGlobalDataTable()
+{
+	ISkillSystemModule::Get().GetSkillSystemGlobals().AutomationTestOnly_SetGlobalAttributeDataTable(NULL);
+}
+
 bool GameplayEffectsTest_InstantDamage_ScalingExplicit(UWorld *World, FAutomationTestBase * Test)
 {
 	// This example uses explicit scaling in a GameplayEffect. We explicitly specify the curve table to use in the GameplayEffect
@@ -1870,6 +1957,9 @@ bool GameplayEffectsTest_InstantDamage_ScalingExplicit(UWorld *World, FAutomatio
 		float ExpectedValue = StartHealth - (LevelOfDamage * SourceDamageScale);
 		Test->TestTrue(SKILL_TEST_TEXT("Damage Applied. Health: Actual: %.2f == Exected: %.2f", ActualValue, ExpectedValue), ActualValue == ExpectedValue);
 	}
+
+	World->EditorDestroyActor(SourceActor, false);
+	World->EditorDestroyActor(DestActor, false);
 	
 	return true;
 }
@@ -1915,6 +2005,10 @@ bool GameplayEffectsTest_InstantDamage_ScalingGlobal(UWorld *World, FAutomationT
 	}
 
 	ClearGlobalCurveTable();
+
+	World->EditorDestroyActor(SourceActor, false);
+	World->EditorDestroyActor(DestActor, false);
+
 	return true;
 }
 
@@ -1963,6 +2057,10 @@ bool GameplayEffectsTest_InstantDamage_Override(UWorld *World, FAutomationTestBa
 	}
 
 	ClearGlobalCurveTable();
+
+	World->EditorDestroyActor(SourceActor, false);
+	World->EditorDestroyActor(DestActor, false);
+
 	return true;
 }
 
@@ -2050,6 +2148,9 @@ bool GameplayEffectsTest_InstantDamageRequiredTag(UWorld *World, FAutomationTest
 		Test->TestTrue(SKILL_TEST_TEXT("Damage Applied. Health: Actual: %.2f == Exected: %.2f", ActualValue, ExpectedValue), ActualValue == ExpectedValue);
 	}
 
+	World->EditorDestroyActor(SourceActor, false);
+	World->EditorDestroyActor(DestActor, false);
+
 	return true;
 }
 
@@ -2137,6 +2238,9 @@ bool GameplayEffectsTest_InstantDamageIgnoreTag(UWorld *World, FAutomationTestBa
 		SKILL_LOG(Log, TEXT("Final Health: %.2f"), DestComponent->GetSet<USkillSystemTestAttributeSet>()->Health);
 	}
 
+	World->EditorDestroyActor(SourceActor, false);
+	World->EditorDestroyActor(DestActor, false);
+
 	return true;
 }
 
@@ -2216,6 +2320,9 @@ bool GameplayEffectsTest_InstantDamageModifierPassesTag(UWorld *World, FAutomati
 		SKILL_LOG(Log, TEXT("Final Health: %.2f"), DestComponent->GetSet<USkillSystemTestAttributeSet>()->Health);
 	}
 
+	World->EditorDestroyActor(SourceActor, false);
+	World->EditorDestroyActor(DestActor, false);
+
 	return true;
 }
 
@@ -2294,6 +2401,9 @@ bool GameplayEffectsTest_InstantDamageModifierTag(UWorld *World, FAutomationTest
 		Test->TestTrue(SKILL_TEST_TEXT("Buff Instant Damage Applied.  Actual: %.2f == Exected: %.2f", ActualValue, ExpectedValue), (ActualValue == ExpectedValue));
 	}
 
+	World->EditorDestroyActor(SourceActor, false);
+	World->EditorDestroyActor(DestActor, false);
+
 	return true;
 }
 
@@ -2347,6 +2457,10 @@ bool GameplayEffectsTest_InstantDamage_ScalingProperty(UWorld *World, FAutomatio
 	}
 
 	ClearGlobalCurveTable();
+
+	World->EditorDestroyActor(SourceActor, false);
+	World->EditorDestroyActor(DestActor, false);
+
 	return true;
 }
 
@@ -2398,6 +2512,10 @@ bool GameplayEffectsTest_InstantDamage_ScalingPropertyNested(UWorld *World, FAut
 	}
 
 	ClearGlobalCurveTable();
+
+	World->EditorDestroyActor(SourceActor, false);
+	World->EditorDestroyActor(DestActor, false);
+
 	return true;
 }
 
@@ -2446,11 +2564,6 @@ bool GameplayEffectsTest_DotDamage_ScalingProperty_Snapshot(UWorld *World, FAuto
 		BaseDmgEffect->LevelInfo.TakeSnapshotOnInit = true;								// But just a snapshot of their SpellDamage when we are applied
 
 		SourceComponent->ApplyGameplayEffectSpecToTarget(BaseDmgEffect, DestComponent);
-
-		// Confirm we ticked once
-		float ActualValue = DestComponent->GetSet<USkillSystemTestAttributeSet>()->Health;
-		float ExpectedValue = StartHealth - (SpellDamage * GameplayEffectScaling);
-		Test->TestTrue(SKILL_TEST_TEXT("Damage Applied. Health: Actual: %.2f == Exected: %.2f", ActualValue, ExpectedValue), ActualValue == ExpectedValue);
 	}
 
 
@@ -2473,11 +2586,15 @@ bool GameplayEffectsTest_DotDamage_ScalingProperty_Snapshot(UWorld *World, FAuto
 		GameplayTest_TickWorld(World, 1.f);
 		
 		float ActualValue = DestComponent->GetSet<USkillSystemTestAttributeSet>()->Health;
-		float ExpectedValue = StartHealth - (SpellDamage * GameplayEffectScaling) * 2.f;	// We've ticked twice
+		float ExpectedValue = StartHealth - (SpellDamage * GameplayEffectScaling) * 2.f;	// we trigger twice on this tick
 		Test->TestTrue(SKILL_TEST_TEXT("Damage Applied. Health: Actual: %.2f == Exected: %.2f", ActualValue, ExpectedValue), ActualValue == ExpectedValue);
 	}	
 
 	ClearGlobalCurveTable();
+
+	World->EditorDestroyActor(SourceActor, false);
+	World->EditorDestroyActor(DestActor, false);
+
 	return true;
 }
 
@@ -2527,15 +2644,19 @@ bool GameplayEffectsTest_DotDamage_ScalingProperty_Dynamic(UWorld *World, FAutom
 
 		SourceComponent->ApplyGameplayEffectSpecToTarget(BaseDmgEffect, DestComponent);
 
-		// Confirm we ticked once
-		float ActualValue = DestComponent->GetSet<USkillSystemTestAttributeSet>()->Health;
-		float ExpectedValue = StartHealth - (SpellDamage * GameplayEffectScaling);
-		Test->TestTrue(SKILL_TEST_TEXT("Damage Applied. Health: Actual: %.2f == Exected: %.2f", ActualValue, ExpectedValue), ActualValue == ExpectedValue);
-
 		float SpellDamageTest = SourceComponent->GetSet<USkillSystemTestAttributeSet>()->SpellDamage;
 		check(SpellDamageTest == SpellDamage);
 	}
 
+	{
+		// Tick once
+		GameplayTest_TickWorld(World, 1.f);
+
+		float ActualValue = DestComponent->GetSet<USkillSystemTestAttributeSet>()->Health;
+		float ExpectedValue = StartHealth - (SpellDamage * GameplayEffectScaling * 2);	// We've ticked twice
+		Test->TestTrue(SKILL_TEST_TEXT("Damage Applied. Health: Actual: %.2f == Exected: %.2f", ActualValue, ExpectedValue), ActualValue == ExpectedValue);
+	}
+
 	{
 		// Increase spell damage on instigator (after we already applied the DOT)
 		UGameplayEffect * SpellDmgEffect = Cast<UGameplayEffect>(StaticConstructObject(UGameplayEffect::StaticClass(), GetTransientPackage(), FName(TEXT("SpellDmgEffect"))));
@@ -2559,11 +2680,15 @@ bool GameplayEffectsTest_DotDamage_ScalingProperty_Dynamic(UWorld *World, FAutom
 		GameplayTest_TickWorld(World, 1.f);
 
 		float ActualValue = DestComponent->GetSet<USkillSystemTestAttributeSet>()->Health;
-		float ExpectedValue = StartHealth - (SpellDamage * GameplayEffectScaling) - (SpellDamage2 * GameplayEffectScaling);	// We've ticked twice
+		float ExpectedValue = StartHealth - (SpellDamage * GameplayEffectScaling * 2) - (SpellDamage2 * GameplayEffectScaling);	// two unbuffed ticks, one buffed tick
 		Test->TestTrue(SKILL_TEST_TEXT("Damage Applied. Health: Actual: %.2f == Exected: %.2f", ActualValue, ExpectedValue), ActualValue == ExpectedValue);
 	}
 
 	ClearGlobalCurveTable();
+
+	World->EditorDestroyActor(SourceActor, false);
+	World->EditorDestroyActor(DestActor, false);
+
 	return true;
 }
 
@@ -2642,6 +2767,9 @@ bool GameplayEffectsTest_MetaAttributes(UWorld *World, FAutomationTestBase * Tes
 	}
 
 	ClearGlobalCurveTable();
+
+	World->EditorDestroyActor(SourceActor, false);
+
 	return true;
 }
 
@@ -2727,12 +2855,1143 @@ bool GameplayEffectsTest_TagOrdering(UWorld *World, FAutomationTestBase * Test)
 		Test->TestTrue(SKILL_TEST_TEXT("MaxHealth: Actual: %.2f == Exected: %.2f", ActualValue, ExpectedValue), ActualValue == ExpectedValue);
 	}
 
+	World->EditorDestroyActor(SourceActor, false);
+	World->EditorDestroyActor(DestActor, false);
+
+	return true;
+}
+
+// test the rule that stacks based on the highest gameplay effect
+bool GameplayEffectsTest_StackingHighest(UWorld *World, FAutomationTestBase * Test)
+{
+	const float StackingValue = 5.f;
+
+	SetGlobalCurveTable();
+	SetGlobalDataTable();
+
+	ASkillSystemTestPawn *SourceActor = World->SpawnActor<ASkillSystemTestPawn>();
+	ASkillSystemTestPawn *DestActor = World->SpawnActor<ASkillSystemTestPawn>();
+
+	UProperty *StackingProperty = FindFieldChecked<UProperty>(USkillSystemTestAttributeSet::StaticClass(), GET_MEMBER_NAME_CHECKED(USkillSystemTestAttributeSet, StackingAttribute1));
+
+	UAttributeComponent * SourceComponent = SourceActor->AttributeComponent;
+	UAttributeComponent * DestComponent = DestActor->AttributeComponent;
+
+	{
+		SKILL_LOG_SCOPE(TEXT("Apply StackingEffect"));
+
+		UGameplayEffect * BaseStackedEffect = Cast<UGameplayEffect>(StaticConstructObject(UGameplayEffect::StaticClass(), GetTransientPackage(), FName(TEXT("StackingEffect1"))));
+		BaseStackedEffect->Modifiers.SetNum(1);
+		BaseStackedEffect->Modifiers[0].Magnitude.SetValue(StackingValue);
+		BaseStackedEffect->Modifiers[0].ModifierType = EGameplayMod::Attribute;
+		BaseStackedEffect->Modifiers[0].ModifierOp = EGameplayModOp::Additive;
+		BaseStackedEffect->Modifiers[0].Attribute.SetUProperty(StackingProperty);
+		BaseStackedEffect->Duration.Value = UGameplayEffect::INFINITE_DURATION;
+		BaseStackedEffect->Period.Value = 1.f;
+		BaseStackedEffect->StackingPolicy = EGameplayEffectStackingPolicy::Highest;
+
+		SourceComponent->ApplyGameplayEffectSpecToTarget(BaseStackedEffect, DestComponent);
+	}
+
+	{
+		// Tick once
+		GameplayTest_TickWorld(World, 1.f);
+
+		float ActualValue = DestComponent->GetSet<USkillSystemTestAttributeSet>()->StackingAttribute1;
+		float ExpectedValue = StackingValue * 2.f;	// effect will execute twice
+		Test->TestTrue(SKILL_TEST_TEXT("Stacking Attribute: Actual: %.2f == Exected: %.2f", ActualValue, ExpectedValue), ActualValue == ExpectedValue);
+	}
+
+	{
+		SKILL_LOG_SCOPE(TEXT("Apply StackingEffect"));
+
+		UGameplayEffect * BaseStackedEffect = Cast<UGameplayEffect>(StaticConstructObject(UGameplayEffect::StaticClass(), GetTransientPackage(), FName(TEXT("StackingEffect2"))));
+		BaseStackedEffect->Modifiers.SetNum(1);
+		BaseStackedEffect->Modifiers[0].Magnitude.SetValue(StackingValue * 2.f);
+		BaseStackedEffect->Modifiers[0].ModifierType = EGameplayMod::Attribute;
+		BaseStackedEffect->Modifiers[0].ModifierOp = EGameplayModOp::Additive;
+		BaseStackedEffect->Modifiers[0].Attribute.SetUProperty(StackingProperty);
+		BaseStackedEffect->Duration.Value = UGameplayEffect::INFINITE_DURATION;
+		BaseStackedEffect->Period.Value = 1.f;
+		BaseStackedEffect->StackingPolicy = EGameplayEffectStackingPolicy::Highest;
+
+		SourceComponent->ApplyGameplayEffectSpecToTarget(BaseStackedEffect, DestComponent);
+	}
+
+	{
+		// Tick once
+		GameplayTest_TickWorld(World, 1.f);
+
+		float ActualValue = DestComponent->GetSet<USkillSystemTestAttributeSet>()->StackingAttribute1;
+		float ExpectedValue = StackingValue * 6.f;	// 2 for first effect, 4 for second effect executing twice
+		Test->TestTrue(SKILL_TEST_TEXT("Stacking Attribute: Actual: %.2f == Exected: %.2f", ActualValue, ExpectedValue), ActualValue == ExpectedValue);
+	}
+
+	ClearGlobalCurveTable();
+	ClearGlobalDataTable();
+
+	World->EditorDestroyActor(SourceActor, false);
+	World->EditorDestroyActor(DestActor, false);
+
+	return true;
+}
+
+// test the rule that stacks based on the smallest gameplay effect
+bool GameplayEffectsTest_StackingLowest(UWorld *World, FAutomationTestBase * Test)
+{
+	const float StackingValue = 5.f;
+
+	SetGlobalCurveTable();
+	SetGlobalDataTable();
+
+	ASkillSystemTestPawn *SourceActor = World->SpawnActor<ASkillSystemTestPawn>();
+	ASkillSystemTestPawn *DestActor = World->SpawnActor<ASkillSystemTestPawn>();
+
+	UProperty *StackingProperty = FindFieldChecked<UProperty>(USkillSystemTestAttributeSet::StaticClass(), GET_MEMBER_NAME_CHECKED(USkillSystemTestAttributeSet, StackingAttribute1));
+
+	UAttributeComponent * SourceComponent = SourceActor->AttributeComponent;
+	UAttributeComponent * DestComponent = DestActor->AttributeComponent;
+
+	{
+		SKILL_LOG_SCOPE(TEXT("Apply StackingEffect"));
+
+		UGameplayEffect * BaseStackedEffect = Cast<UGameplayEffect>(StaticConstructObject(UGameplayEffect::StaticClass(), GetTransientPackage(), FName(TEXT("StackingEffect1"))));
+		BaseStackedEffect->Modifiers.SetNum(1);
+		BaseStackedEffect->Modifiers[0].Magnitude.SetValue(StackingValue);
+		BaseStackedEffect->Modifiers[0].ModifierType = EGameplayMod::Attribute;
+		BaseStackedEffect->Modifiers[0].ModifierOp = EGameplayModOp::Additive;
+		BaseStackedEffect->Modifiers[0].Attribute.SetUProperty(StackingProperty);
+		BaseStackedEffect->Duration.Value = UGameplayEffect::INFINITE_DURATION;
+		BaseStackedEffect->Period.Value = 1.f;
+		BaseStackedEffect->StackingPolicy = EGameplayEffectStackingPolicy::Lowest;
+
+		SourceComponent->ApplyGameplayEffectSpecToTarget(BaseStackedEffect, DestComponent);
+	}
+
+	{
+		// Tick once
+		GameplayTest_TickWorld(World, 1.f);
+
+		float ActualValue = DestComponent->GetSet<USkillSystemTestAttributeSet>()->StackingAttribute1;
+		float ExpectedValue = StackingValue * 2.f;	// the effect should execute twice here
+		Test->TestTrue(SKILL_TEST_TEXT("Stacking Attribute: Actual: %.2f == Exected: %.2f", ActualValue, ExpectedValue), ActualValue == ExpectedValue);
+	}
+
+	{
+		SKILL_LOG_SCOPE(TEXT("Apply StackingEffect"));
+
+		UGameplayEffect * BaseStackedEffect = Cast<UGameplayEffect>(StaticConstructObject(UGameplayEffect::StaticClass(), GetTransientPackage(), FName(TEXT("StackingEffect2"))));
+		BaseStackedEffect->Modifiers.SetNum(1);
+		BaseStackedEffect->Modifiers[0].Magnitude.SetValue(StackingValue * 2.f);
+		BaseStackedEffect->Modifiers[0].ModifierType = EGameplayMod::Attribute;
+		BaseStackedEffect->Modifiers[0].ModifierOp = EGameplayModOp::Additive;
+		BaseStackedEffect->Modifiers[0].Attribute.SetUProperty(StackingProperty);
+		BaseStackedEffect->Duration.Value = UGameplayEffect::INFINITE_DURATION;
+		BaseStackedEffect->Period.Value = 1.f;
+		BaseStackedEffect->StackingPolicy = EGameplayEffectStackingPolicy::Lowest;
+
+		SourceComponent->ApplyGameplayEffectSpecToTarget(BaseStackedEffect, DestComponent);
+	}
+
+	{
+		// Tick once
+		GameplayTest_TickWorld(World, 1.f);
+
+		float ActualValue = DestComponent->GetSet<USkillSystemTestAttributeSet>()->StackingAttribute1;
+		float ExpectedValue = StackingValue * 3.f;	// the first effect has executed 3 times, the second hasn't executed
+		Test->TestTrue(SKILL_TEST_TEXT("Stacking Attribute: Actual: %.2f == Exected: %.2f", ActualValue, ExpectedValue), ActualValue == ExpectedValue);
+	}
+
+	ClearGlobalCurveTable();
+	ClearGlobalDataTable();
+
+	World->EditorDestroyActor(SourceActor, false);
+	World->EditorDestroyActor(DestActor, false);
+
+	return true;
+}
+
+// test the rule that ignores stacking
+bool GameplayEffectsTest_StackingUnlimited(UWorld *World, FAutomationTestBase * Test)
+{
+	const float StackingValue = 5.f;
+
+	SetGlobalCurveTable();
+	SetGlobalDataTable();
+
+	ASkillSystemTestPawn *SourceActor = World->SpawnActor<ASkillSystemTestPawn>();
+	ASkillSystemTestPawn *DestActor = World->SpawnActor<ASkillSystemTestPawn>();
+
+	UProperty *StackingProperty = FindFieldChecked<UProperty>(USkillSystemTestAttributeSet::StaticClass(), GET_MEMBER_NAME_CHECKED(USkillSystemTestAttributeSet, StackingAttribute1));
+
+	UAttributeComponent * SourceComponent = SourceActor->AttributeComponent;
+	UAttributeComponent * DestComponent = DestActor->AttributeComponent;
+
+	{
+		SKILL_LOG_SCOPE(TEXT("Apply StackingEffect"));
+
+		UGameplayEffect * BaseStackedEffect = Cast<UGameplayEffect>(StaticConstructObject(UGameplayEffect::StaticClass(), GetTransientPackage(), FName(TEXT("StackingEffect1"))));
+		BaseStackedEffect->Modifiers.SetNum(1);
+		BaseStackedEffect->Modifiers[0].Magnitude.SetValue(StackingValue);
+		BaseStackedEffect->Modifiers[0].ModifierType = EGameplayMod::Attribute;
+		BaseStackedEffect->Modifiers[0].ModifierOp = EGameplayModOp::Additive;
+		BaseStackedEffect->Modifiers[0].Attribute.SetUProperty(StackingProperty);
+		BaseStackedEffect->Duration.Value = UGameplayEffect::INFINITE_DURATION;
+		BaseStackedEffect->Period.Value = 1.f;
+		BaseStackedEffect->StackingPolicy = EGameplayEffectStackingPolicy::Unlimited;
+
+		SourceComponent->ApplyGameplayEffectSpecToTarget(BaseStackedEffect, DestComponent);
+	}
+
+	{
+		SKILL_LOG_SCOPE(TEXT("Apply StackingEffect"));
+
+		UGameplayEffect * BaseStackedEffect = Cast<UGameplayEffect>(StaticConstructObject(UGameplayEffect::StaticClass(), GetTransientPackage(), FName(TEXT("StackingEffect2"))));
+		BaseStackedEffect->Modifiers.SetNum(1);
+		BaseStackedEffect->Modifiers[0].Magnitude.SetValue(StackingValue * 2.f);
+		BaseStackedEffect->Modifiers[0].ModifierType = EGameplayMod::Attribute;
+		BaseStackedEffect->Modifiers[0].ModifierOp = EGameplayModOp::Additive;
+		BaseStackedEffect->Modifiers[0].Attribute.SetUProperty(StackingProperty);
+		BaseStackedEffect->Duration.Value = UGameplayEffect::INFINITE_DURATION;
+		BaseStackedEffect->Period.Value = 1.f;
+		BaseStackedEffect->StackingPolicy = EGameplayEffectStackingPolicy::Unlimited;
+
+		SourceComponent->ApplyGameplayEffectSpecToTarget(BaseStackedEffect, DestComponent);
+	}
+
+	{
+		// Tick once
+		GameplayTest_TickWorld(World, 1.f);
+
+		float ActualValue = DestComponent->GetSet<USkillSystemTestAttributeSet>()->StackingAttribute1;
+		float ExpectedValue = StackingValue * 6.f;	// 1 for first GE, 2 for the second GE, 3 for both GEs during the tick
+		Test->TestTrue(SKILL_TEST_TEXT("Stacking Attribute: Actual: %.2f == Exected: %.2f", ActualValue, ExpectedValue), ActualValue == ExpectedValue);
+	}
+
+	ClearGlobalCurveTable();
+	ClearGlobalDataTable();
+
+	World->EditorDestroyActor(SourceActor, false);
+	World->EditorDestroyActor(DestActor, false);
+
+	return true;
+}
+
+// test to make sure that stacking updates correctly when a gameplay effect is removed
+bool GameplayEffectsTest_StackingRemoval(UWorld *World, FAutomationTestBase * Test)
+{
+	const float StackingValue = 5.f;
+
+	SetGlobalCurveTable();
+	SetGlobalDataTable();
+
+	ASkillSystemTestPawn *SourceActor = World->SpawnActor<ASkillSystemTestPawn>();
+	ASkillSystemTestPawn *DestActor = World->SpawnActor<ASkillSystemTestPawn>();
+
+	UProperty *StackingProperty = FindFieldChecked<UProperty>(USkillSystemTestAttributeSet::StaticClass(), GET_MEMBER_NAME_CHECKED(USkillSystemTestAttributeSet, StackingAttribute1));
+
+	UAttributeComponent * SourceComponent = SourceActor->AttributeComponent;
+	UAttributeComponent * DestComponent = DestActor->AttributeComponent;
+
+	{
+		SKILL_LOG_SCOPE(TEXT("Apply StackingEffect"));
+
+		UGameplayEffect * BaseStackedEffect = Cast<UGameplayEffect>(StaticConstructObject(UGameplayEffect::StaticClass(), GetTransientPackage(), FName(TEXT("StackingEffect1"))));
+		BaseStackedEffect->Modifiers.SetNum(1);
+		BaseStackedEffect->Modifiers[0].Magnitude.SetValue(StackingValue);
+		BaseStackedEffect->Modifiers[0].ModifierType = EGameplayMod::Attribute;
+		BaseStackedEffect->Modifiers[0].ModifierOp = EGameplayModOp::Additive;
+		BaseStackedEffect->Modifiers[0].Attribute.SetUProperty(StackingProperty);
+		BaseStackedEffect->Duration.Value = UGameplayEffect::INFINITE_DURATION;
+		BaseStackedEffect->Period.Value = 1.f;
+		BaseStackedEffect->StackingPolicy = EGameplayEffectStackingPolicy::Highest;
+
+		SourceComponent->ApplyGameplayEffectSpecToTarget(BaseStackedEffect, DestComponent);
+	}
+
+	{
+		SKILL_LOG_SCOPE(TEXT("Apply StackingEffect"));
+
+		UGameplayEffect * BaseStackedEffect = Cast<UGameplayEffect>(StaticConstructObject(UGameplayEffect::StaticClass(), GetTransientPackage(), FName(TEXT("StackingEffect2"))));
+		BaseStackedEffect->Modifiers.SetNum(1);
+		BaseStackedEffect->Modifiers[0].Magnitude.SetValue(StackingValue * 2.f);
+		BaseStackedEffect->Modifiers[0].ModifierType = EGameplayMod::Attribute;
+		BaseStackedEffect->Modifiers[0].ModifierOp = EGameplayModOp::Additive;
+		BaseStackedEffect->Modifiers[0].Attribute.SetUProperty(StackingProperty);
+		BaseStackedEffect->Duration.Value = 1.f;// UGameplayEffect::INFINITE_DURATION;
+		BaseStackedEffect->Period.Value = 1.f;
+		BaseStackedEffect->StackingPolicy = EGameplayEffectStackingPolicy::Highest;
+
+		SourceComponent->ApplyGameplayEffectSpecToTarget(BaseStackedEffect, DestComponent);
+	}
+
+	{
+		// Tick once
+		GameplayTest_TickWorld(World, 1.f);
+
+		float ActualValue = DestComponent->GetSet<USkillSystemTestAttributeSet>()->StackingAttribute1;
+		float ExpectedValue = StackingValue * 4.f;	// 2 for second effect, 2 for second effect during tick
+		Test->TestTrue(SKILL_TEST_TEXT("Stacking Attribute: Actual: %.2f == Exected: %.2f", ActualValue, ExpectedValue), ActualValue == ExpectedValue);
+
+		// Tick once
+		GameplayTest_TickWorld(World, 1.f);
+
+		ActualValue = DestComponent->GetSet<USkillSystemTestAttributeSet>()->StackingAttribute1;
+		ExpectedValue = StackingValue * 5.f;	// 2 for second effect, 2 for second effect during first tick, 1 for first effect during second tick
+		Test->TestTrue(SKILL_TEST_TEXT("Stacking Attribute: Actual: %.2f == Exected: %.2f", ActualValue, ExpectedValue), ActualValue == ExpectedValue);
+	}
+
+	ClearGlobalCurveTable();
+	ClearGlobalDataTable();
+
+	World->EditorDestroyActor(SourceActor, false);
+	World->EditorDestroyActor(DestActor, false);
+
+	return true;
+}
+
+// test that attributes that aren't allowed to stack won't stack even if the stacking rule says they should
+bool GameplayEffectsTest_StackingNoStack(UWorld *World, FAutomationTestBase * Test)
+{
+	const float StackingValue = 5.f;
+
+	SetGlobalCurveTable();
+	SetGlobalDataTable();
+
+	ASkillSystemTestPawn *SourceActor = World->SpawnActor<ASkillSystemTestPawn>();
+	ASkillSystemTestPawn *DestActor = World->SpawnActor<ASkillSystemTestPawn>();
+
+	UProperty *NoStackProperty = FindFieldChecked<UProperty>(USkillSystemTestAttributeSet::StaticClass(), GET_MEMBER_NAME_CHECKED(USkillSystemTestAttributeSet, NoStackAttribute));
+
+	UAttributeComponent * SourceComponent = SourceActor->AttributeComponent;
+	UAttributeComponent * DestComponent = DestActor->AttributeComponent;
+
+	{
+		SKILL_LOG_SCOPE(TEXT("Apply StackingEffect"));
+
+		UGameplayEffect * UnstackableEffect = Cast<UGameplayEffect>(StaticConstructObject(UGameplayEffect::StaticClass(), GetTransientPackage(), FName(TEXT("NoStackEffect1"))));
+		UnstackableEffect->Modifiers.SetNum(1);
+		UnstackableEffect->Modifiers[0].Magnitude.SetValue(StackingValue);
+		UnstackableEffect->Modifiers[0].ModifierType = EGameplayMod::Attribute;
+		UnstackableEffect->Modifiers[0].ModifierOp = EGameplayModOp::Additive;
+		UnstackableEffect->Modifiers[0].Attribute.SetUProperty(NoStackProperty);
+		UnstackableEffect->Duration.Value = UGameplayEffect::INFINITE_DURATION;
+		UnstackableEffect->Period.Value = 1.f;
+		UnstackableEffect->StackingPolicy = EGameplayEffectStackingPolicy::Highest;
+
+		SourceComponent->ApplyGameplayEffectSpecToTarget(UnstackableEffect, DestComponent);
+	}
+
+	{
+		SKILL_LOG_SCOPE(TEXT("Apply StackingEffect"));
+
+		UGameplayEffect * UnstackableEffect = Cast<UGameplayEffect>(StaticConstructObject(UGameplayEffect::StaticClass(), GetTransientPackage(), FName(TEXT("NoStackEffect2"))));
+		UnstackableEffect->Modifiers.SetNum(1);
+		UnstackableEffect->Modifiers[0].Magnitude.SetValue(StackingValue * 2.f);
+		UnstackableEffect->Modifiers[0].ModifierType = EGameplayMod::Attribute;
+		UnstackableEffect->Modifiers[0].ModifierOp = EGameplayModOp::Additive;
+		UnstackableEffect->Modifiers[0].Attribute.SetUProperty(NoStackProperty);
+		UnstackableEffect->Duration.Value = UGameplayEffect::INFINITE_DURATION;
+		UnstackableEffect->Period.Value = 1.f;
+		UnstackableEffect->StackingPolicy = EGameplayEffectStackingPolicy::Highest;
+
+		SourceComponent->ApplyGameplayEffectSpecToTarget(UnstackableEffect, DestComponent);
+	}
+
+	{
+		// Tick once
+		GameplayTest_TickWorld(World, 1.f);
+
+		float ActualValue = DestComponent->GetSet<USkillSystemTestAttributeSet>()->NoStackAttribute;
+		float ExpectedValue = StackingValue * 6.f;	// 1 for first effect, 2 for second effect, 3 for both effects during tick
+		Test->TestTrue(SKILL_TEST_TEXT("Stacking Attribute: Actual: %.2f == Exected: %.2f", ActualValue, ExpectedValue), ActualValue == ExpectedValue);
+	}
+
+	ClearGlobalCurveTable();
+	ClearGlobalDataTable();
+
+	World->EditorDestroyActor(SourceActor, false);
+	World->EditorDestroyActor(DestActor, false);
+
+	return true;
+}
+
+// test the custom rule that stacks based on capping the number of effects applied
+bool GameplayEffectsTest_StackingCustomCapped(UWorld *World, FAutomationTestBase * Test)
+{
+	const float StackingValue = 5.f;
+
+	SetGlobalCurveTable();
+	SetGlobalDataTable();
+
+	ASkillSystemTestPawn *SourceActor = World->SpawnActor<ASkillSystemTestPawn>();
+	ASkillSystemTestPawn *DestActor = World->SpawnActor<ASkillSystemTestPawn>();
+
+	UProperty *StackingProperty = FindFieldChecked<UProperty>(USkillSystemTestAttributeSet::StaticClass(), GET_MEMBER_NAME_CHECKED(USkillSystemTestAttributeSet, StackingAttribute1));
+
+	UAttributeComponent * SourceComponent = SourceActor->AttributeComponent;
+	UAttributeComponent * DestComponent = DestActor->AttributeComponent;
+
+	{
+		SKILL_LOG_SCOPE(TEXT("Apply StackingEffect"));
+
+		UGameplayEffect * BaseStackedEffect = Cast<UGameplayEffect>(StaticConstructObject(UGameplayEffect::StaticClass(), GetTransientPackage(), FName(TEXT("StackingEffect1"))));
+		BaseStackedEffect->Modifiers.SetNum(1);
+		BaseStackedEffect->Modifiers[0].Magnitude.SetValue(StackingValue);
+		BaseStackedEffect->Modifiers[0].ModifierType = EGameplayMod::Attribute;
+		BaseStackedEffect->Modifiers[0].ModifierOp = EGameplayModOp::Additive;
+		BaseStackedEffect->Modifiers[0].Attribute.SetUProperty(StackingProperty);
+		BaseStackedEffect->Modifiers[0].OwnedTags.AddTag("Stackable");
+		BaseStackedEffect->Duration.Value = 2.f;
+		BaseStackedEffect->Period.Value = 1.f;
+		BaseStackedEffect->StackingPolicy = EGameplayEffectStackingPolicy::Callback;
+		BaseStackedEffect->StackingExtension = UGameplayEffectStackingExtension_CappedNumberTest::StaticClass();
+//		BaseStackedEffect->StackingExtension = UGameplayEffectStackingExtension_DiminishingReturnsTest::StaticClass();
+
+		SourceComponent->ApplyGameplayEffectSpecToTarget(BaseStackedEffect, DestComponent);
+	}
+
+	{
+		SKILL_LOG_SCOPE(TEXT("Apply StackingEffect"));
+
+		UGameplayEffect * BaseStackedEffect = Cast<UGameplayEffect>(StaticConstructObject(UGameplayEffect::StaticClass(), GetTransientPackage(), FName(TEXT("StackingEffect2"))));
+		BaseStackedEffect->Modifiers.SetNum(1);
+		BaseStackedEffect->Modifiers[0].Magnitude.SetValue(StackingValue);
+		BaseStackedEffect->Modifiers[0].ModifierType = EGameplayMod::Attribute;
+		BaseStackedEffect->Modifiers[0].ModifierOp = EGameplayModOp::Additive;
+		BaseStackedEffect->Modifiers[0].Attribute.SetUProperty(StackingProperty);
+		BaseStackedEffect->Modifiers[0].OwnedTags.AddTag("Stackable");
+		BaseStackedEffect->Duration.Value = 2.f;
+		BaseStackedEffect->Period.Value = 1.f;
+		BaseStackedEffect->StackingPolicy = EGameplayEffectStackingPolicy::Callback;
+		BaseStackedEffect->StackingExtension = UGameplayEffectStackingExtension_CappedNumberTest::StaticClass();
+
+		SourceComponent->ApplyGameplayEffectSpecToTarget(BaseStackedEffect, DestComponent);
+	}
+
+	{
+		SKILL_LOG_SCOPE(TEXT("Apply StackingEffect"));
+
+		UGameplayEffect * BaseStackedEffect = Cast<UGameplayEffect>(StaticConstructObject(UGameplayEffect::StaticClass(), GetTransientPackage(), FName(TEXT("StackingEffect3"))));
+		BaseStackedEffect->Modifiers.SetNum(1);
+		BaseStackedEffect->Modifiers[0].Magnitude.SetValue(StackingValue);
+		BaseStackedEffect->Modifiers[0].ModifierType = EGameplayMod::Attribute;
+		BaseStackedEffect->Modifiers[0].ModifierOp = EGameplayModOp::Additive;
+		BaseStackedEffect->Modifiers[0].Attribute.SetUProperty(StackingProperty);
+		BaseStackedEffect->Modifiers[0].OwnedTags.AddTag("Stackable");
+		BaseStackedEffect->Duration.Value = 2.f;
+		BaseStackedEffect->Period.Value = 1.f;
+		BaseStackedEffect->StackingPolicy = EGameplayEffectStackingPolicy::Callback;
+		BaseStackedEffect->StackingExtension = UGameplayEffectStackingExtension_CappedNumberTest::StaticClass();
+
+		SourceComponent->ApplyGameplayEffectSpecToTarget(BaseStackedEffect, DestComponent);
+	}
+
+	{
+		// Tick once
+		GameplayTest_TickWorld(World, 1.f);
+
+		float ActualValue = DestComponent->GetSet<USkillSystemTestAttributeSet>()->StackingAttribute1;
+		float ExpectedValue = StackingValue * 4.f;	// 2 for effects being applied at the start of the tick, 2 for effects being applied at the end of this tick
+		Test->TestTrue(SKILL_TEST_TEXT("Stacking Attribute: Actual: %.2f == Exected: %.2f", ActualValue, ExpectedValue), ActualValue == ExpectedValue);
+	}
+
+	{
+		SKILL_LOG_SCOPE(TEXT("Apply StackingEffect"));
+
+		UGameplayEffect * BaseStackedEffect = Cast<UGameplayEffect>(StaticConstructObject(UGameplayEffect::StaticClass(), GetTransientPackage(), FName(TEXT("StackingEffect4"))));
+		BaseStackedEffect->Modifiers.SetNum(1);
+		BaseStackedEffect->Modifiers[0].Magnitude.SetValue(StackingValue);
+		BaseStackedEffect->Modifiers[0].ModifierType = EGameplayMod::Attribute;
+		BaseStackedEffect->Modifiers[0].ModifierOp = EGameplayModOp::Additive;
+		BaseStackedEffect->Modifiers[0].Attribute.SetUProperty(StackingProperty);
+		BaseStackedEffect->Modifiers[0].OwnedTags.AddTag("Stackable");
+		BaseStackedEffect->Duration.Value = 2.f;
+		BaseStackedEffect->Period.Value = 1.f;
+		BaseStackedEffect->StackingPolicy = EGameplayEffectStackingPolicy::Callback;
+		BaseStackedEffect->StackingExtension = UGameplayEffectStackingExtension_CappedNumberTest::StaticClass();
+
+		SourceComponent->ApplyGameplayEffectSpecToTarget(BaseStackedEffect, DestComponent);
+	}
+
+	{
+		// Tick once
+		GameplayTest_TickWorld(World, 1.f);
+
+		float ActualValue = DestComponent->GetSet<USkillSystemTestAttributeSet>()->StackingAttribute1;
+		float ExpectedValue = StackingValue * 6.f;	// 2 for effects being applied, 2 for tick, 2 for tick
+		Test->TestTrue(SKILL_TEST_TEXT("Stacking Attribute: Actual: %.2f == Exected: %.2f", ActualValue, ExpectedValue), ActualValue == ExpectedValue);
+	}
+
+	{
+		// Tick once
+		GameplayTest_TickWorld(World, 1.f);
+
+		float ActualValue = DestComponent->GetSet<USkillSystemTestAttributeSet>()->StackingAttribute1;
+		float ExpectedValue = StackingValue * 8.f;	// the last GE should have refreshed the timer so we should have 2 for GEs applied and 3 * 2 for ticks
+		Test->TestTrue(SKILL_TEST_TEXT("Stacking Attribute: Actual: %.2f == Exected: %.2f", ActualValue, ExpectedValue), ActualValue == ExpectedValue);
+	}
+
+	{
+		// Tick once
+		GameplayTest_TickWorld(World, 1.f);
+
+		float ActualValue = DestComponent->GetSet<USkillSystemTestAttributeSet>()->StackingAttribute1;
+		float ExpectedValue = StackingValue * 8.f;	// the effects should have timed out so we should have 2 for GEs applied and 3 * 2 for ticks
+		Test->TestTrue(SKILL_TEST_TEXT("Stacking Attribute: Actual: %.2f == Exected: %.2f", ActualValue, ExpectedValue), ActualValue == ExpectedValue);
+	}
+
+	ClearGlobalCurveTable();
+	ClearGlobalDataTable();
+
+	World->EditorDestroyActor(SourceActor, false);
+	World->EditorDestroyActor(DestActor, false);
+
+	return true;
+}
+
+// test the custom rule that stacks based on diminishing returns
+bool GameplayEffectsTest_StackingCustomDiminishingReturns(UWorld *World, FAutomationTestBase * Test)
+{
+	const float StackingValue = 1.f;
+
+	SetGlobalCurveTable();
+	SetGlobalDataTable();
+
+	ASkillSystemTestPawn *SourceActor = World->SpawnActor<ASkillSystemTestPawn>();
+	ASkillSystemTestPawn *DestActor = World->SpawnActor<ASkillSystemTestPawn>();
+
+	UProperty *StackingProperty = FindFieldChecked<UProperty>(USkillSystemTestAttributeSet::StaticClass(), GET_MEMBER_NAME_CHECKED(USkillSystemTestAttributeSet, StackingAttribute1));
+
+	UAttributeComponent * SourceComponent = SourceActor->AttributeComponent;
+	UAttributeComponent * DestComponent = DestActor->AttributeComponent;
+
+	{
+		SKILL_LOG_SCOPE(TEXT("Apply StackingEffect"));
+
+		UGameplayEffect * BaseStackedEffect = Cast<UGameplayEffect>(StaticConstructObject(UGameplayEffect::StaticClass(), GetTransientPackage(), FName(TEXT("StackingEffect1"))));
+		BaseStackedEffect->Modifiers.SetNum(1);
+		BaseStackedEffect->Modifiers[0].Magnitude.SetValue(StackingValue);
+		BaseStackedEffect->Modifiers[0].ModifierType = EGameplayMod::Attribute;
+		BaseStackedEffect->Modifiers[0].ModifierOp = EGameplayModOp::Additive;
+		BaseStackedEffect->Modifiers[0].Attribute.SetUProperty(StackingProperty);
+		BaseStackedEffect->Modifiers[0].OwnedTags.AddTag("Stackable");
+		BaseStackedEffect->Duration.Value = UGameplayEffect::INFINITE_DURATION;
+		BaseStackedEffect->Period.Value = 1.f;
+		BaseStackedEffect->StackingPolicy = EGameplayEffectStackingPolicy::Callback;
+		BaseStackedEffect->StackingExtension = UGameplayEffectStackingExtension_DiminishingReturnsTest::StaticClass();
+
+		SourceComponent->ApplyGameplayEffectSpecToTarget(BaseStackedEffect, DestComponent);
+	}
+
+	{
+		// Tick once
+		GameplayTest_TickWorld(World, 1.f);
+
+		float ActualValue = DestComponent->GetSet<USkillSystemTestAttributeSet>()->StackingAttribute1;
+		float ExpectedValue = StackingValue * (5.f * 2);	// first application gets five times the result and will be applied at the start and end of the tick
+		Test->TestTrue(SKILL_TEST_TEXT("Stacking Attribute: Actual: %.2f == Exected: %.2f", ActualValue, ExpectedValue), ActualValue == ExpectedValue);
+	}
+
+	{
+		SKILL_LOG_SCOPE(TEXT("Apply StackingEffect"));
+
+		UGameplayEffect * BaseStackedEffect = Cast<UGameplayEffect>(StaticConstructObject(UGameplayEffect::StaticClass(), GetTransientPackage(), FName(TEXT("StackingEffect2"))));
+		BaseStackedEffect->Modifiers.SetNum(1);
+		BaseStackedEffect->Modifiers[0].Magnitude.SetValue(StackingValue);
+		BaseStackedEffect->Modifiers[0].ModifierType = EGameplayMod::Attribute;
+		BaseStackedEffect->Modifiers[0].ModifierOp = EGameplayModOp::Additive;
+		BaseStackedEffect->Modifiers[0].Attribute.SetUProperty(StackingProperty);
+		BaseStackedEffect->Modifiers[0].OwnedTags.AddTag("Stackable");
+		BaseStackedEffect->Duration.Value = UGameplayEffect::INFINITE_DURATION;
+		BaseStackedEffect->Period.Value = 1.f;
+		BaseStackedEffect->StackingPolicy = EGameplayEffectStackingPolicy::Callback;
+		BaseStackedEffect->StackingExtension = UGameplayEffectStackingExtension_DiminishingReturnsTest::StaticClass();
+
+		SourceComponent->ApplyGameplayEffectSpecToTarget(BaseStackedEffect, DestComponent);
+	}
+
+	{
+		// Tick once
+		GameplayTest_TickWorld(World, 1.f);
+
+		float ActualValue = DestComponent->GetSet<USkillSystemTestAttributeSet>()->StackingAttribute1;
+		float ExpectedValue = StackingValue * ((5.f * 2) + 7.f);	// second application gets seven times the result
+		Test->TestTrue(SKILL_TEST_TEXT("Stacking Attribute: Actual: %.2f == Exected: %.2f", ActualValue, ExpectedValue), ActualValue == ExpectedValue);
+	}
+
+	{
+		SKILL_LOG_SCOPE(TEXT("Apply StackingEffect"));
+
+		UGameplayEffect * BaseStackedEffect = Cast<UGameplayEffect>(StaticConstructObject(UGameplayEffect::StaticClass(), GetTransientPackage(), FName(TEXT("StackingEffect3"))));
+		BaseStackedEffect->Modifiers.SetNum(1);
+		BaseStackedEffect->Modifiers[0].Magnitude.SetValue(StackingValue);
+		BaseStackedEffect->Modifiers[0].ModifierType = EGameplayMod::Attribute;
+		BaseStackedEffect->Modifiers[0].ModifierOp = EGameplayModOp::Additive;
+		BaseStackedEffect->Modifiers[0].Attribute.SetUProperty(StackingProperty);
+		BaseStackedEffect->Modifiers[0].OwnedTags.AddTag("Stackable");
+		BaseStackedEffect->Duration.Value = UGameplayEffect::INFINITE_DURATION;
+		BaseStackedEffect->Period.Value = 1.f;
+		BaseStackedEffect->StackingPolicy = EGameplayEffectStackingPolicy::Callback;
+		BaseStackedEffect->StackingExtension = UGameplayEffectStackingExtension_DiminishingReturnsTest::StaticClass();
+
+		SourceComponent->ApplyGameplayEffectSpecToTarget(BaseStackedEffect, DestComponent);
+	}
+
+	{
+		// Tick once
+		GameplayTest_TickWorld(World, 1.f);
+
+		float ActualValue = DestComponent->GetSet<USkillSystemTestAttributeSet>()->StackingAttribute1;
+		float ExpectedValue = StackingValue * ((5.f * 2) + 7.f + 8.f);	// third application gets eight times the result
+		Test->TestTrue(SKILL_TEST_TEXT("Stacking Attribute: Actual: %.2f == Exected: %.2f", ActualValue, ExpectedValue), ActualValue == ExpectedValue);
+	}
+
+	{
+		SKILL_LOG_SCOPE(TEXT("Apply StackingEffect"));
+
+		UGameplayEffect * BaseStackedEffect = Cast<UGameplayEffect>(StaticConstructObject(UGameplayEffect::StaticClass(), GetTransientPackage(), FName(TEXT("StackingEffect4"))));
+		BaseStackedEffect->Modifiers.SetNum(1);
+		BaseStackedEffect->Modifiers[0].Magnitude.SetValue(StackingValue);
+		BaseStackedEffect->Modifiers[0].ModifierType = EGameplayMod::Attribute;
+		BaseStackedEffect->Modifiers[0].ModifierOp = EGameplayModOp::Additive;
+		BaseStackedEffect->Modifiers[0].Attribute.SetUProperty(StackingProperty);
+		BaseStackedEffect->Modifiers[0].OwnedTags.AddTag("Stackable");
+		BaseStackedEffect->Duration.Value = UGameplayEffect::INFINITE_DURATION;
+		BaseStackedEffect->Period.Value = 1.f;
+		BaseStackedEffect->StackingPolicy = EGameplayEffectStackingPolicy::Callback;
+		BaseStackedEffect->StackingExtension = UGameplayEffectStackingExtension_DiminishingReturnsTest::StaticClass();
+
+		SourceComponent->ApplyGameplayEffectSpecToTarget(BaseStackedEffect, DestComponent);
+	}
+
+	{
+		// Tick once
+		GameplayTest_TickWorld(World, 1.f);
+
+		float ActualValue = DestComponent->GetSet<USkillSystemTestAttributeSet>()->StackingAttribute1;
+		float ExpectedValue = StackingValue * ((5.f * 2) + 7.f + 8.f + 9.f);	// fourth application gets nine times the result
+		Test->TestTrue(SKILL_TEST_TEXT("Stacking Attribute: Actual: %.2f == Exected: %.2f", ActualValue, ExpectedValue), ActualValue == ExpectedValue);
+	}
+
+	ClearGlobalCurveTable();
+	ClearGlobalDataTable();
+
+	World->EditorDestroyActor(SourceActor, false);
+	World->EditorDestroyActor(DestActor, false);
+
+	return true;
+}
+
+// Gameplay effects that change the same attribute but have different stacking rules shouldn't interfere with each other
+bool GameplayEffectsTest_StackingDifferentRules(UWorld *World, FAutomationTestBase * Test)
+{
+	const float StackingValue = 5.f;
+
+	SetGlobalCurveTable();
+	SetGlobalDataTable();
+
+	ASkillSystemTestPawn *SourceActor = World->SpawnActor<ASkillSystemTestPawn>();
+	ASkillSystemTestPawn *DestActor = World->SpawnActor<ASkillSystemTestPawn>();
+
+	UProperty *StackingProperty = FindFieldChecked<UProperty>(USkillSystemTestAttributeSet::StaticClass(), GET_MEMBER_NAME_CHECKED(USkillSystemTestAttributeSet, StackingAttribute1));
+
+	UAttributeComponent * SourceComponent = SourceActor->AttributeComponent;
+	UAttributeComponent * DestComponent = DestActor->AttributeComponent;
+
+	{
+		SKILL_LOG_SCOPE(TEXT("Apply StackingEffect"));
+
+		UGameplayEffect * BaseStackedEffect = Cast<UGameplayEffect>(StaticConstructObject(UGameplayEffect::StaticClass(), GetTransientPackage(), FName(TEXT("StackingEffect1"))));
+		BaseStackedEffect->Modifiers.SetNum(1);
+		BaseStackedEffect->Modifiers[0].Magnitude.SetValue(StackingValue);
+		BaseStackedEffect->Modifiers[0].ModifierType = EGameplayMod::Attribute;
+		BaseStackedEffect->Modifiers[0].ModifierOp = EGameplayModOp::Additive;
+		BaseStackedEffect->Modifiers[0].Attribute.SetUProperty(StackingProperty);
+		BaseStackedEffect->Duration.Value = UGameplayEffect::INFINITE_DURATION;
+		BaseStackedEffect->Period.Value = 1.f;
+		BaseStackedEffect->StackingPolicy = EGameplayEffectStackingPolicy::Highest;
+
+		SourceComponent->ApplyGameplayEffectSpecToTarget(BaseStackedEffect, DestComponent);
+	}
+
+	{
+		SKILL_LOG_SCOPE(TEXT("Apply StackingEffect"));
+
+		UGameplayEffect * BaseStackedEffect = Cast<UGameplayEffect>(StaticConstructObject(UGameplayEffect::StaticClass(), GetTransientPackage(), FName(TEXT("StackingEffect2"))));
+		BaseStackedEffect->Modifiers.SetNum(1);
+		BaseStackedEffect->Modifiers[0].Magnitude.SetValue(StackingValue * 2.f);
+		BaseStackedEffect->Modifiers[0].ModifierType = EGameplayMod::Attribute;
+		BaseStackedEffect->Modifiers[0].ModifierOp = EGameplayModOp::Additive;
+		BaseStackedEffect->Modifiers[0].Attribute.SetUProperty(StackingProperty);
+		BaseStackedEffect->Duration.Value = UGameplayEffect::INFINITE_DURATION;
+		BaseStackedEffect->Period.Value = 1.f;
+		BaseStackedEffect->StackingPolicy = EGameplayEffectStackingPolicy::Lowest;
+
+		SourceComponent->ApplyGameplayEffectSpecToTarget(BaseStackedEffect, DestComponent);
+	}
+
+	{
+		// Tick once
+		GameplayTest_TickWorld(World, 1.f);
+
+		float ActualValue = DestComponent->GetSet<USkillSystemTestAttributeSet>()->StackingAttribute1;
+		float ExpectedValue = StackingValue * 6.f;	// 1 for first effect, 2 for the second effect, 3 during tick
+		Test->TestTrue(SKILL_TEST_TEXT("Stacking Attribute: Actual: %.2f == Exected: %.2f", ActualValue, ExpectedValue), ActualValue == ExpectedValue);
+	}
+
+	ClearGlobalCurveTable();
+	ClearGlobalDataTable();
+
+	World->EditorDestroyActor(SourceActor, false);
+	World->EditorDestroyActor(DestActor, false);
+
+	return true;
+}
+
+// different attributes with the same stacking rule shouldn't interfere with each other
+bool GameplayEffectsTest_StackingDifferentAttributes(UWorld *World, FAutomationTestBase * Test)
+{
+	const float StackingValue = 5.f;
+
+	SetGlobalCurveTable();
+	SetGlobalDataTable();
+
+	ASkillSystemTestPawn *SourceActor = World->SpawnActor<ASkillSystemTestPawn>();
+	ASkillSystemTestPawn *DestActor = World->SpawnActor<ASkillSystemTestPawn>();
+
+	UProperty *StackingProperty1 = FindFieldChecked<UProperty>(USkillSystemTestAttributeSet::StaticClass(), GET_MEMBER_NAME_CHECKED(USkillSystemTestAttributeSet, StackingAttribute1));
+	UProperty *StackingProperty2 = FindFieldChecked<UProperty>(USkillSystemTestAttributeSet::StaticClass(), GET_MEMBER_NAME_CHECKED(USkillSystemTestAttributeSet, StackingAttribute2));
+
+	UAttributeComponent * SourceComponent = SourceActor->AttributeComponent;
+	UAttributeComponent * DestComponent = DestActor->AttributeComponent;
+
+	{
+		SKILL_LOG_SCOPE(TEXT("Apply StackingEffect"));
+
+		UGameplayEffect * BaseStackedEffect = Cast<UGameplayEffect>(StaticConstructObject(UGameplayEffect::StaticClass(), GetTransientPackage(), FName(TEXT("StackingEffect1"))));
+		BaseStackedEffect->Modifiers.SetNum(1);
+		BaseStackedEffect->Modifiers[0].Magnitude.SetValue(StackingValue);
+		BaseStackedEffect->Modifiers[0].ModifierType = EGameplayMod::Attribute;
+		BaseStackedEffect->Modifiers[0].ModifierOp = EGameplayModOp::Additive;
+		BaseStackedEffect->Modifiers[0].Attribute.SetUProperty(StackingProperty1);
+		BaseStackedEffect->Duration.Value = UGameplayEffect::INFINITE_DURATION;
+		BaseStackedEffect->Period.Value = 1.f;
+		BaseStackedEffect->StackingPolicy = EGameplayEffectStackingPolicy::Highest;
+
+		SourceComponent->ApplyGameplayEffectSpecToTarget(BaseStackedEffect, DestComponent);
+	}
+
+	{
+		SKILL_LOG_SCOPE(TEXT("Apply StackingEffect"));
+
+		UGameplayEffect * BaseStackedEffect = Cast<UGameplayEffect>(StaticConstructObject(UGameplayEffect::StaticClass(), GetTransientPackage(), FName(TEXT("StackingEffect2"))));
+		BaseStackedEffect->Modifiers.SetNum(1);
+		BaseStackedEffect->Modifiers[0].Magnitude.SetValue(StackingValue * 2.f);
+		BaseStackedEffect->Modifiers[0].ModifierType = EGameplayMod::Attribute;
+		BaseStackedEffect->Modifiers[0].ModifierOp = EGameplayModOp::Additive;
+		BaseStackedEffect->Modifiers[0].Attribute.SetUProperty(StackingProperty2);
+		BaseStackedEffect->Duration.Value = UGameplayEffect::INFINITE_DURATION;
+		BaseStackedEffect->Period.Value = 1.f;
+		BaseStackedEffect->StackingPolicy = EGameplayEffectStackingPolicy::Highest;
+
+		SourceComponent->ApplyGameplayEffectSpecToTarget(BaseStackedEffect, DestComponent);
+	}
+
+	{
+		// Tick once
+		GameplayTest_TickWorld(World, 1.f);
+
+		float ActualValue = DestComponent->GetSet<USkillSystemTestAttributeSet>()->StackingAttribute1;
+		float ExpectedValue = StackingValue * 2.f;	// 1 for first effect, 1 for first effect during tick
+		Test->TestTrue(SKILL_TEST_TEXT("Stacking Attribute 1: %.2f == Exected: %.2f", ActualValue, ExpectedValue), ActualValue == ExpectedValue);
+		ActualValue = DestComponent->GetSet<USkillSystemTestAttributeSet>()->StackingAttribute2;
+		ExpectedValue = StackingValue * 4.f; // 2 for the second effect, 2 for the second effect during tick
+		Test->TestTrue(SKILL_TEST_TEXT("Stacking Attribute 2: Actual: %.2f == Exected: %.2f", ActualValue, ExpectedValue), ActualValue == ExpectedValue);
+	}
+
+	ClearGlobalCurveTable();
+	ClearGlobalDataTable();
+
+	World->EditorDestroyActor(SourceActor, false);
+	World->EditorDestroyActor(DestActor, false);
+
+	return true;
+}
+
+// test to make sure that two different custom rules don't interfere with each other
+bool GameplayEffectsTest_StackingCustomTwoRules(UWorld *World, FAutomationTestBase * Test)
+{
+	const float StackingValue = 1.f;
+
+	SetGlobalCurveTable();
+	SetGlobalDataTable();
+
+	ASkillSystemTestPawn *SourceActor = World->SpawnActor<ASkillSystemTestPawn>();
+	ASkillSystemTestPawn *DestActor = World->SpawnActor<ASkillSystemTestPawn>();
+
+	UProperty *StackingProperty = FindFieldChecked<UProperty>(USkillSystemTestAttributeSet::StaticClass(), GET_MEMBER_NAME_CHECKED(USkillSystemTestAttributeSet, StackingAttribute1));
+
+	UAttributeComponent * SourceComponent = SourceActor->AttributeComponent;
+	UAttributeComponent * DestComponent = DestActor->AttributeComponent;
+
+	{
+		SKILL_LOG_SCOPE(TEXT("Apply StackingEffect"));
+
+		UGameplayEffect * BaseStackedEffect = Cast<UGameplayEffect>(StaticConstructObject(UGameplayEffect::StaticClass(), GetTransientPackage(), FName(TEXT("StackingEffect1"))));
+		BaseStackedEffect->Modifiers.SetNum(1);
+		BaseStackedEffect->Modifiers[0].Magnitude.SetValue(StackingValue);
+		BaseStackedEffect->Modifiers[0].ModifierType = EGameplayMod::Attribute;
+		BaseStackedEffect->Modifiers[0].ModifierOp = EGameplayModOp::Additive;
+		BaseStackedEffect->Modifiers[0].Attribute.SetUProperty(StackingProperty);
+		BaseStackedEffect->Modifiers[0].OwnedTags.AddTag("Stackable");
+		BaseStackedEffect->Duration.Value = UGameplayEffect::INFINITE_DURATION;
+		BaseStackedEffect->Period.Value = 1.f;
+		BaseStackedEffect->StackingPolicy = EGameplayEffectStackingPolicy::Callback;
+		BaseStackedEffect->StackingExtension = UGameplayEffectStackingExtension_CappedNumberTest::StaticClass();
+
+		SourceComponent->ApplyGameplayEffectSpecToTarget(BaseStackedEffect, DestComponent);
+	}
+
+	{
+		SKILL_LOG_SCOPE(TEXT("Apply StackingEffect"));
+
+		UGameplayEffect * BaseStackedEffect = Cast<UGameplayEffect>(StaticConstructObject(UGameplayEffect::StaticClass(), GetTransientPackage(), FName(TEXT("StackingEffect2"))));
+		BaseStackedEffect->Modifiers.SetNum(1);
+		BaseStackedEffect->Modifiers[0].Magnitude.SetValue(StackingValue);
+		BaseStackedEffect->Modifiers[0].ModifierType = EGameplayMod::Attribute;
+		BaseStackedEffect->Modifiers[0].ModifierOp = EGameplayModOp::Additive;
+		BaseStackedEffect->Modifiers[0].Attribute.SetUProperty(StackingProperty);
+		BaseStackedEffect->Modifiers[0].OwnedTags.AddTag("Stackable");
+		BaseStackedEffect->Duration.Value = UGameplayEffect::INFINITE_DURATION;
+		BaseStackedEffect->Period.Value = 1.f;
+		BaseStackedEffect->StackingPolicy = EGameplayEffectStackingPolicy::Callback;
+		BaseStackedEffect->StackingExtension = UGameplayEffectStackingExtension_DiminishingReturnsTest::StaticClass();
+
+		SourceComponent->ApplyGameplayEffectSpecToTarget(BaseStackedEffect, DestComponent);
+	}
+
+	{
+		SKILL_LOG_SCOPE(TEXT("Apply StackingEffect"));
+
+		UGameplayEffect * BaseStackedEffect = Cast<UGameplayEffect>(StaticConstructObject(UGameplayEffect::StaticClass(), GetTransientPackage(), FName(TEXT("StackingEffect3"))));
+		BaseStackedEffect->Modifiers.SetNum(1);
+		BaseStackedEffect->Modifiers[0].Magnitude.SetValue(StackingValue);
+		BaseStackedEffect->Modifiers[0].ModifierType = EGameplayMod::Attribute;
+		BaseStackedEffect->Modifiers[0].ModifierOp = EGameplayModOp::Additive;
+		BaseStackedEffect->Modifiers[0].Attribute.SetUProperty(StackingProperty);
+		BaseStackedEffect->Modifiers[0].OwnedTags.AddTag("Stackable");
+		BaseStackedEffect->Duration.Value = UGameplayEffect::INFINITE_DURATION;
+		BaseStackedEffect->Period.Value = 1.f;
+		BaseStackedEffect->StackingPolicy = EGameplayEffectStackingPolicy::Callback;
+		BaseStackedEffect->StackingExtension = UGameplayEffectStackingExtension_DiminishingReturnsTest::StaticClass();
+
+		SourceComponent->ApplyGameplayEffectSpecToTarget(BaseStackedEffect, DestComponent);
+	}
+
+	{
+		// Tick once
+		GameplayTest_TickWorld(World, 1.f);
+
+		float ActualValue = DestComponent->GetSet<USkillSystemTestAttributeSet>()->StackingAttribute1;
+		float ExpectedValue = StackingValue * ((1 + 7) * 2);	// the capped stacking rule will apply one, the diminishing returns rule will apply seven, both rules are applied twice
+		Test->TestTrue(SKILL_TEST_TEXT("Stacking Attribute: Actual: %.2f == Exected: %.2f", ActualValue, ExpectedValue), ActualValue == ExpectedValue);
+	}
+
+	ClearGlobalCurveTable();
+	ClearGlobalDataTable();
+
+	World->EditorDestroyActor(SourceActor, false);
+	World->EditorDestroyActor(DestActor, false);
+
+	return true;
+}
+
+// test to make sure that a custom rules applied to different attributes doesn't interfere with itself
+bool GameplayEffectsTest_StackingCustomTwoAttributes(UWorld *World, FAutomationTestBase * Test)
+{
+	const float StackingValue = 1.f;
+
+	SetGlobalCurveTable();
+	SetGlobalDataTable();
+
+	ASkillSystemTestPawn *SourceActor = World->SpawnActor<ASkillSystemTestPawn>();
+	ASkillSystemTestPawn *DestActor = World->SpawnActor<ASkillSystemTestPawn>();
+
+	UProperty *StackingProperty1 = FindFieldChecked<UProperty>(USkillSystemTestAttributeSet::StaticClass(), GET_MEMBER_NAME_CHECKED(USkillSystemTestAttributeSet, StackingAttribute1));
+	UProperty *StackingProperty2 = FindFieldChecked<UProperty>(USkillSystemTestAttributeSet::StaticClass(), GET_MEMBER_NAME_CHECKED(USkillSystemTestAttributeSet, StackingAttribute2));
+
+	UAttributeComponent * SourceComponent = SourceActor->AttributeComponent;
+	UAttributeComponent * DestComponent = DestActor->AttributeComponent;
+
+	{
+		SKILL_LOG_SCOPE(TEXT("Apply StackingEffect"));
+
+		UGameplayEffect * BaseStackedEffect = Cast<UGameplayEffect>(StaticConstructObject(UGameplayEffect::StaticClass(), GetTransientPackage(), FName(TEXT("StackingEffect1"))));
+		BaseStackedEffect->Modifiers.SetNum(1);
+		BaseStackedEffect->Modifiers[0].Magnitude.SetValue(StackingValue);
+		BaseStackedEffect->Modifiers[0].ModifierType = EGameplayMod::Attribute;
+		BaseStackedEffect->Modifiers[0].ModifierOp = EGameplayModOp::Additive;
+		BaseStackedEffect->Modifiers[0].Attribute.SetUProperty(StackingProperty1);
+		BaseStackedEffect->Modifiers[0].OwnedTags.AddTag("Stackable");
+		BaseStackedEffect->Duration.Value = UGameplayEffect::INFINITE_DURATION;
+		BaseStackedEffect->Period.Value = 1.f;
+		BaseStackedEffect->StackingPolicy = EGameplayEffectStackingPolicy::Callback;
+		BaseStackedEffect->StackingExtension = UGameplayEffectStackingExtension_CappedNumberTest::StaticClass();
+
+		SourceComponent->ApplyGameplayEffectSpecToTarget(BaseStackedEffect, DestComponent);
+	}
+
+	{
+		SKILL_LOG_SCOPE(TEXT("Apply StackingEffect"));
+
+		UGameplayEffect * BaseStackedEffect = Cast<UGameplayEffect>(StaticConstructObject(UGameplayEffect::StaticClass(), GetTransientPackage(), FName(TEXT("StackingEffect2"))));
+		BaseStackedEffect->Modifiers.SetNum(1);
+		BaseStackedEffect->Modifiers[0].Magnitude.SetValue(StackingValue);
+		BaseStackedEffect->Modifiers[0].ModifierType = EGameplayMod::Attribute;
+		BaseStackedEffect->Modifiers[0].ModifierOp = EGameplayModOp::Additive;
+		BaseStackedEffect->Modifiers[0].Attribute.SetUProperty(StackingProperty1);
+		BaseStackedEffect->Modifiers[0].OwnedTags.AddTag("Stackable");
+		BaseStackedEffect->Duration.Value = UGameplayEffect::INFINITE_DURATION;
+		BaseStackedEffect->Period.Value = 1.f;
+		BaseStackedEffect->StackingPolicy = EGameplayEffectStackingPolicy::Callback;
+		BaseStackedEffect->StackingExtension = UGameplayEffectStackingExtension_CappedNumberTest::StaticClass();
+
+		SourceComponent->ApplyGameplayEffectSpecToTarget(BaseStackedEffect, DestComponent);
+	}
+
+	{
+		SKILL_LOG_SCOPE(TEXT("Apply StackingEffect"));
+
+		UGameplayEffect * BaseStackedEffect = Cast<UGameplayEffect>(StaticConstructObject(UGameplayEffect::StaticClass(), GetTransientPackage(), FName(TEXT("StackingEffect3"))));
+		BaseStackedEffect->Modifiers.SetNum(1);
+		BaseStackedEffect->Modifiers[0].Magnitude.SetValue(StackingValue);
+		BaseStackedEffect->Modifiers[0].ModifierType = EGameplayMod::Attribute;
+		BaseStackedEffect->Modifiers[0].ModifierOp = EGameplayModOp::Additive;
+		BaseStackedEffect->Modifiers[0].Attribute.SetUProperty(StackingProperty2);
+		BaseStackedEffect->Modifiers[0].OwnedTags.AddTag("Stackable");
+		BaseStackedEffect->Duration.Value = UGameplayEffect::INFINITE_DURATION;
+		BaseStackedEffect->Period.Value = 1.f;
+		BaseStackedEffect->StackingPolicy = EGameplayEffectStackingPolicy::Callback;
+		BaseStackedEffect->StackingExtension = UGameplayEffectStackingExtension_CappedNumberTest::StaticClass();
+
+		SourceComponent->ApplyGameplayEffectSpecToTarget(BaseStackedEffect, DestComponent);
+	}
+
+	{
+		// Tick once
+		GameplayTest_TickWorld(World, 1.f);
+
+		float ActualValue = DestComponent->GetSet<USkillSystemTestAttributeSet>()->StackingAttribute1;
+		float ExpectedValue = StackingValue * (2 * 2);	// two effects should be applied to the first attribute, both apply twice
+		Test->TestTrue(SKILL_TEST_TEXT("Stacking Attribute: Actual: %.2f == Exected: %.2f", ActualValue, ExpectedValue), ActualValue == ExpectedValue);
+
+		ActualValue = DestComponent->GetSet<USkillSystemTestAttributeSet>()->StackingAttribute2;
+		ExpectedValue = StackingValue * (1 * 2);	// one effect should be applied twice to the second attribute
+		Test->TestTrue(SKILL_TEST_TEXT("Stacking Attribute: Actual: %.2f == Exected: %.2f", ActualValue, ExpectedValue), ActualValue == ExpectedValue);
+	}
+
+	ClearGlobalCurveTable();
+	ClearGlobalDataTable();
+
+	World->EditorDestroyActor(SourceActor, false);
+	World->EditorDestroyActor(DestActor, false);
+
+	return true;
+}
+
+// Test to make sure that removing a modifier causes stacks to be recalculated correctly
+bool GameplayEffectsTest_StackingRemovingModifiers(UWorld *World, FAutomationTestBase * Test)
+{
+	const float StackingValue = 5.f;
+
+	SetGlobalCurveTable();
+	SetGlobalDataTable();
+
+	ASkillSystemTestPawn *SourceActor = World->SpawnActor<ASkillSystemTestPawn>();
+	ASkillSystemTestPawn *DestActor = World->SpawnActor<ASkillSystemTestPawn>();
+
+	UProperty *StackingProperty = FindFieldChecked<UProperty>(USkillSystemTestAttributeSet::StaticClass(), GET_MEMBER_NAME_CHECKED(USkillSystemTestAttributeSet, StackingAttribute1));
+
+	UAttributeComponent * SourceComponent = SourceActor->AttributeComponent;
+	UAttributeComponent * DestComponent = DestActor->AttributeComponent;
+
+	// Setup a GE to modify IncomingGEs
+	{
+		SKILL_LOG_SCOPE(TEXT("Apply modifier to incoming, tagged GEs"));
+
+		UGameplayEffect* BaseModEffect = Cast<UGameplayEffect>(StaticConstructObject(UGameplayEffect::StaticClass(), GetTransientPackage(), FName(TEXT("ProtectBuff"))));
+		BaseModEffect->Modifiers.SetNum(1);
+		BaseModEffect->Modifiers[0].Magnitude.SetValue(4.f);
+		BaseModEffect->Modifiers[0].ModifierType = EGameplayMod::IncomingGE;
+		BaseModEffect->Modifiers[0].ModifierOp = EGameplayModOp::Multiplicitive;
+		BaseModEffect->Modifiers[0].Attribute.SetUProperty(StackingProperty);
+		BaseModEffect->Duration.SetValue(1.f);
+		BaseModEffect->CopyPolicy = EGameplayEffectCopyPolicy::AlwaysLink;
+		BaseModEffect->GameplayEffectRequiredTags.AddTag(FName(TEXT("Stack")));
+
+		// Apply to self
+		DestComponent->ApplyGameplayEffectSpecToTarget(BaseModEffect, DestComponent, 1.f);
+	}
+
+	{
+		SKILL_LOG_SCOPE(TEXT("Apply StackingEffect"));
+
+		UGameplayEffect * BaseStackedEffect = Cast<UGameplayEffect>(StaticConstructObject(UGameplayEffect::StaticClass(), GetTransientPackage(), FName(TEXT("StackingEffect1"))));
+		BaseStackedEffect->Modifiers.SetNum(1);
+		BaseStackedEffect->Modifiers[0].Magnitude.SetValue(StackingValue);
+		BaseStackedEffect->Modifiers[0].ModifierType = EGameplayMod::Attribute;
+		BaseStackedEffect->Modifiers[0].ModifierOp = EGameplayModOp::Additive;
+		BaseStackedEffect->Modifiers[0].Attribute.SetUProperty(StackingProperty);
+		BaseStackedEffect->Duration.Value = UGameplayEffect::INFINITE_DURATION;
+		BaseStackedEffect->Period.Value = 1.f;
+		BaseStackedEffect->StackingPolicy = EGameplayEffectStackingPolicy::Highest;
+		BaseStackedEffect->GameplayEffectTags.AddTag(FName(TEXT("Stack")));
+
+		SourceComponent->ApplyGameplayEffectSpecToTarget(BaseStackedEffect, DestComponent);
+	}
+
+	{
+		SKILL_LOG_SCOPE(TEXT("Apply StackingEffect"));
+
+		UGameplayEffect * BaseStackedEffect = Cast<UGameplayEffect>(StaticConstructObject(UGameplayEffect::StaticClass(), GetTransientPackage(), FName(TEXT("StackingEffect2"))));
+		BaseStackedEffect->Modifiers.SetNum(1);
+		BaseStackedEffect->Modifiers[0].Magnitude.SetValue(StackingValue * 2.f);
+		BaseStackedEffect->Modifiers[0].ModifierType = EGameplayMod::Attribute;
+		BaseStackedEffect->Modifiers[0].ModifierOp = EGameplayModOp::Additive;
+		BaseStackedEffect->Modifiers[0].Attribute.SetUProperty(StackingProperty);
+		BaseStackedEffect->Duration.Value = UGameplayEffect::INFINITE_DURATION;
+		BaseStackedEffect->Period.Value = 1.f;
+		BaseStackedEffect->StackingPolicy = EGameplayEffectStackingPolicy::Highest;
+
+		SourceComponent->ApplyGameplayEffectSpecToTarget(BaseStackedEffect, DestComponent);
+	}
+
+	{
+		// Tick once
+		GameplayTest_TickWorld(World, 1.f);
+
+		float ActualValue = DestComponent->GetSet<USkillSystemTestAttributeSet>()->StackingAttribute1;
+		float ExpectedValue = StackingValue * 8.f;	// 4 for the first GE, 4 for first GE during tick
+		Test->TestTrue(SKILL_TEST_TEXT("Stacking Attribute: Actual: %.2f == Exected: %.2f", ActualValue, ExpectedValue), ActualValue == ExpectedValue);
+
+		// At this point the modifier should be removed and the second GE should be the best match for the stacking rule
+		// Tick once
+		GameplayTest_TickWorld(World, 1.f);
+
+		ActualValue = DestComponent->GetSet<USkillSystemTestAttributeSet>()->StackingAttribute1;
+		ExpectedValue = StackingValue * 10.f;	// 4 for the first GE, 4 for first GE during first tick, 2 for the second GE during second tick
+		Test->TestTrue(SKILL_TEST_TEXT("Stacking Attribute: Actual: %.2f == Exected: %.2f", ActualValue, ExpectedValue), ActualValue == ExpectedValue);
+	}
+
+	ClearGlobalCurveTable();
+	ClearGlobalDataTable();
+
+	World->EditorDestroyActor(SourceActor, false);
+	World->EditorDestroyActor(DestActor, false);
+
+	return true;
+}
+
+// test to make sure that adding a modifier causes stacks to be recalculated correctly
+bool GameplayEffectsTest_StackingAddingModifiers(UWorld *World, FAutomationTestBase * Test)
+{
+	const float StackingValue = 5.f;
+
+	SetGlobalCurveTable();
+	SetGlobalDataTable();
+
+	ASkillSystemTestPawn *SourceActor = World->SpawnActor<ASkillSystemTestPawn>();
+	ASkillSystemTestPawn *DestActor = World->SpawnActor<ASkillSystemTestPawn>();
+
+	UProperty *StackingProperty = FindFieldChecked<UProperty>(USkillSystemTestAttributeSet::StaticClass(), GET_MEMBER_NAME_CHECKED(USkillSystemTestAttributeSet, StackingAttribute1));
+
+	UAttributeComponent * SourceComponent = SourceActor->AttributeComponent;
+	UAttributeComponent * DestComponent = DestActor->AttributeComponent;
+
+	{
+		SKILL_LOG_SCOPE(TEXT("Apply StackingEffect"));
+
+		UGameplayEffect * BaseStackedEffect = Cast<UGameplayEffect>(StaticConstructObject(UGameplayEffect::StaticClass(), GetTransientPackage(), FName(TEXT("StackingEffect1"))));
+		BaseStackedEffect->Modifiers.SetNum(1);
+		BaseStackedEffect->Modifiers[0].Magnitude.SetValue(StackingValue);
+		BaseStackedEffect->Modifiers[0].ModifierType = EGameplayMod::Attribute;
+		BaseStackedEffect->Modifiers[0].ModifierOp = EGameplayModOp::Additive;
+		BaseStackedEffect->Modifiers[0].Attribute.SetUProperty(StackingProperty);
+		BaseStackedEffect->Duration.Value = UGameplayEffect::INFINITE_DURATION;
+		BaseStackedEffect->Period.Value = 1.f;
+		BaseStackedEffect->StackingPolicy = EGameplayEffectStackingPolicy::Highest;
+		BaseStackedEffect->GameplayEffectTags.AddTag(FName(TEXT("Stack")));
+
+		SourceComponent->ApplyGameplayEffectSpecToTarget(BaseStackedEffect, DestComponent);
+	}
+
+	{
+		SKILL_LOG_SCOPE(TEXT("Apply StackingEffect"));
+
+		UGameplayEffect * BaseStackedEffect = Cast<UGameplayEffect>(StaticConstructObject(UGameplayEffect::StaticClass(), GetTransientPackage(), FName(TEXT("StackingEffect2"))));
+		BaseStackedEffect->Modifiers.SetNum(1);
+		BaseStackedEffect->Modifiers[0].Magnitude.SetValue(StackingValue * 2.f);
+		BaseStackedEffect->Modifiers[0].ModifierType = EGameplayMod::Attribute;
+		BaseStackedEffect->Modifiers[0].ModifierOp = EGameplayModOp::Additive;
+		BaseStackedEffect->Modifiers[0].Attribute.SetUProperty(StackingProperty);
+		BaseStackedEffect->Duration.Value = UGameplayEffect::INFINITE_DURATION;
+		BaseStackedEffect->Period.Value = 1.f;
+		BaseStackedEffect->StackingPolicy = EGameplayEffectStackingPolicy::Highest;
+
+		SourceComponent->ApplyGameplayEffectSpecToTarget(BaseStackedEffect, DestComponent);
+	}
+
+	{
+		// Tick once
+		GameplayTest_TickWorld(World, 1.f);
+
+		float ActualValue = DestComponent->GetSet<USkillSystemTestAttributeSet>()->StackingAttribute1;
+		float ExpectedValue = StackingValue * 4.f;	// 4 for the second GE executing twice
+		Test->TestTrue(SKILL_TEST_TEXT("Stacking Attribute: Actual: %.2f == Exected: %.2f", ActualValue, ExpectedValue), ActualValue == ExpectedValue);
+	}
+
+	// Setup a GE to modify IncomingGEs
+	{
+		SKILL_LOG_SCOPE(TEXT("Apply modifier to incoming, tagged GEs"));
+
+		UGameplayEffect* BaseModEffect = Cast<UGameplayEffect>(StaticConstructObject(UGameplayEffect::StaticClass(), GetTransientPackage(), FName(TEXT("ProtectBuff"))));
+		BaseModEffect->Modifiers.SetNum(1);
+		BaseModEffect->Modifiers[0].Magnitude.SetValue(3.f);
+		BaseModEffect->Modifiers[0].ModifierType = EGameplayMod::ActiveGE;
+		BaseModEffect->Modifiers[0].ModifierOp = EGameplayModOp::Multiplicitive;
+		BaseModEffect->Modifiers[0].Attribute.SetUProperty(StackingProperty);
+		BaseModEffect->Duration.SetValue(UGameplayEffect::INFINITE_DURATION);
+		BaseModEffect->CopyPolicy = EGameplayEffectCopyPolicy::AlwaysLink;
+		BaseModEffect->GameplayEffectRequiredTags.AddTag(FName(TEXT("Stack")));
+
+		// Apply to self
+		DestComponent->ApplyGameplayEffectSpecToTarget(BaseModEffect, DestComponent, 1.f);
+	}
+
+	{
+		// At this point the modifier should be removed and the second GE should be the best match for the stacking rule
+		// Tick once
+		GameplayTest_TickWorld(World, 1.f);
+
+		float ActualValue = DestComponent->GetSet<USkillSystemTestAttributeSet>()->StackingAttribute1;
+		float ExpectedValue = StackingValue * 7.f;	// 2 for the second GE, 2 for second GE during tick, 3 for the modified first GE during second tick
+		Test->TestTrue(SKILL_TEST_TEXT("Stacking Attribute: Actual: %.2f == Exected: %.2f", ActualValue, ExpectedValue), ActualValue == ExpectedValue);
+	}
+
+	ClearGlobalCurveTable();
+	ClearGlobalDataTable();
+
+	World->EditorDestroyActor(SourceActor, false);
+	World->EditorDestroyActor(DestActor, false);
+
 	return true;
 }
 
 bool FGameplayEffectsTest::RunTest( const FString& Parameters )
 {
 	UCurveTable *CurveTable = ISkillSystemModule::Get().GetSkillSystemGlobals().GetGlobalCurveTable();
+	UDataTable *DataTable = ISkillSystemModule::Get().GetSkillSystemGlobals().GetGlobalAttributeDataTable();
 
 	UWorld *World = UWorld::CreateWorld(EWorldType::Game, false);
 	FWorldContext &WorldContext = GEngine->CreateNewWorldContext(EWorldType::Game);
@@ -2762,13 +4021,14 @@ bool FGameplayEffectsTest::RunTest( const FString& Parameters )
 
 	GameplayEffectsTest_LifestealExtension(World, this);
 	
-	GameplayEffectsTest_ShieldlExtension(World, this);
-	GameplayEffectsTest_ShieldlExtensionMultiple(World, this);
+	GameplayEffectsTest_ShieldExtension(World, this);
+	GameplayEffectsTest_ShieldExtensionMultiple(World, this);
 	
 	GameplayEffectsTest_InstantDamage_ScalingExplicit(World, this);
 	GameplayEffectsTest_InstantDamage_ScalingGlobal(World, this);
 	GameplayEffectsTest_InstantDamage_Override(World, this);
 
+	// Tagging tests
 	GameplayEffectsTest_InstantDamageRequiredTag(World, this);
 	GameplayEffectsTest_InstantDamageIgnoreTag(World, this);
 	
@@ -2783,11 +4043,36 @@ bool FGameplayEffectsTest::RunTest( const FString& Parameters )
 	GameplayEffectsTest_MetaAttributes(World, this);
 	GameplayEffectsTest_TagOrdering(World, this);
 	
+	//
+	// Stacking GE tests
+	//
+
+	// basic rules
+	GameplayEffectsTest_StackingHighest(World, this);
+	GameplayEffectsTest_StackingLowest(World, this);
+	GameplayEffectsTest_StackingUnlimited(World, this);
+	GameplayEffectsTest_StackingRemoval(World, this);
+	GameplayEffectsTest_StackingNoStack(World, this);
+
+	// custom rules
+	GameplayEffectsTest_StackingCustomCapped(World, this);
+	GameplayEffectsTest_StackingCustomDiminishingReturns(World, this);
+
+	// interactions between different rules/attributes
+	GameplayEffectsTest_StackingDifferentRules(World, this);
+	GameplayEffectsTest_StackingDifferentAttributes(World, this);
+	GameplayEffectsTest_StackingCustomTwoRules(World, this);
+	GameplayEffectsTest_StackingCustomTwoAttributes(World, this);
+
+	// interactions between stacking and modifiers
+	GameplayEffectsTest_StackingRemovingModifiers(World, this);
+	GameplayEffectsTest_StackingAddingModifiers(World, this);
 
 	GEngine->DestroyWorldContext(World);
 	World->DestroyWorld(false);
 
 	ISkillSystemModule::Get().GetSkillSystemGlobals().AutomationTestOnly_SetGlobalCurveTable(CurveTable);
+	ISkillSystemModule::Get().GetSkillSystemGlobals().AutomationTestOnly_SetGlobalAttributeDataTable(DataTable);
 	return true;
 }
 
diff --git a/Engine/Source/Runtime/SkillSystem/Private/SkillSystemGlobals.cpp b/Engine/Source/Runtime/SkillSystem/Private/SkillSystemGlobals.cpp
--- a/Engine/Source/Runtime/SkillSystem/Private/SkillSystemGlobals.cpp
+++ b/Engine/Source/Runtime/SkillSystem/Private/SkillSystemGlobals.cpp
@@ -1,4 +1,4 @@
-// Copyright 1998-2013 Epic Games, Inc. All Rights Reserved.
+// Copyright 1998-2014 Epic Games, Inc. All Rights Reserved.
 
 #include "SkillSystemModulePrivatePCH.h"
 
@@ -17,7 +17,7 @@ UCurveTable * USkillSystemGlobals::GetGlobalCurveTable()
 		// Hook into notifications for object re-imports so that the gameplay tag tree can be reconstructed if the table changes
 		if (GIsEditor && GlobalCurveTable)
 		{
-			GEditor->OnObjectReimported().AddUObject(this, &USkillSystemGlobals::OnObjectReimported);
+			GEditor->OnObjectReimported().AddUObject(this, &USkillSystemGlobals::OnCurveTableReimported);
 		}
 #endif
 	}
@@ -25,13 +25,37 @@ UCurveTable * USkillSystemGlobals::GetGlobalCurveTable()
 	return GlobalCurveTable;
 }
 
+UDataTable * USkillSystemGlobals::GetGlobalAttributeDataTable()
+{
+	if (!GlobalAttributeDataTable && !GlobalAttributeDataTableName.IsEmpty())
+	{
+		GlobalAttributeDataTable = LoadObject<UDataTable>(NULL, *GlobalAttributeDataTableName, NULL, LOAD_None, NULL);
 #if WITH_EDITOR
+		// Hook into notifications for object re-imports so that the gameplay tag tree can be reconstructed if the table changes
+		if (GIsEditor && GlobalAttributeDataTable)
+		{
+			GEditor->OnObjectReimported().AddUObject(this, &USkillSystemGlobals::OnDataTableReimported);
+		}
+#endif
+	}
+
+	return GlobalAttributeDataTable;
+}
 
-void USkillSystemGlobals::OnObjectReimported(UObject* InObject)
+#if WITH_EDITOR
+
+void USkillSystemGlobals::OnCurveTableReimported(UObject* InObject)
 {
 	if (GIsEditor && !IsRunningCommandlet() && InObject && InObject == GlobalCurveTable)
 	{
 	}
 }
 
-#endif
\ No newline at end of file
+void USkillSystemGlobals::OnDataTableReimported(UObject* InObject)
+{
+	if (GIsEditor && !IsRunningCommandlet() && InObject && InObject == GlobalAttributeDataTable)
+	{
+	}
+}
+
+#endif
diff --git a/Engine/Source/Runtime/SkillSystem/Private/SkillSystemTestAttributeSet.cpp b/Engine/Source/Runtime/SkillSystem/Private/SkillSystemTestAttributeSet.cpp
--- a/Engine/Source/Runtime/SkillSystem/Private/SkillSystemTestAttributeSet.cpp
+++ b/Engine/Source/Runtime/SkillSystem/Private/SkillSystemTestAttributeSet.cpp
@@ -11,6 +11,9 @@ USkillSystemTestAttributeSet::USkillSystemTestAttributeSet(const class FPostCons
 	SpellDamage = 0.f;
 	PhysicalDamage = 0.f;
 	Strength = 0.f;
+	StackingAttribute1 = 0.f;
+	StackingAttribute2 = 0.f;
+	NoStackAttribute = 0.f;
 }
 
 
