SkillSystem GameplayAbilities WIP

diff --git a/Engine/Source/Editor/SkillSystemEditor/Classes/K2Node_PlayMontageAndWait.h b/Engine/Source/Editor/SkillSystemEditor/Classes/K2Node_PlayMontageAndWait.h
new file mode 100644
--- /dev/null
+++ b/Engine/Source/Editor/SkillSystemEditor/Classes/K2Node_PlayMontageAndWait.h
@@ -0,0 +1,17 @@
+// Copyright 1998-2014 Epic Games, Inc. All Rights Reserved.
+
+#pragma once
+#include "K2Node_PlayMontageAndWait.generated.h"
+
+UCLASS()
+class UK2Node_PlayMontageAndWait : public UK2Node_BaseAsyncTask
+{
+	GENERATED_UCLASS_BODY()
+
+	// Begin UEdGraphNode interface
+	virtual FString GetTooltip() const OVERRIDE;
+	virtual FText GetNodeTitle(ENodeTitleType::Type TitleType) const OVERRIDE;
+	// End UEdGraphNode interface
+
+	virtual FString GetCategoryName();
+};
diff --git a/Engine/Source/Editor/SkillSystemEditor/Private/K2Node_PlayMontageAndWait.cpp b/Engine/Source/Editor/SkillSystemEditor/Private/K2Node_PlayMontageAndWait.cpp
new file mode 100644
--- /dev/null
+++ b/Engine/Source/Editor/SkillSystemEditor/Private/K2Node_PlayMontageAndWait.cpp
@@ -0,0 +1,37 @@
+// Copyright 1998-2014 Epic Games, Inc. All Rights Reserved.
+
+
+#include "SkillSystemEditorModulePrivatePCH.h"
+#include "SkillSystemBlueprintLibrary.h"
+#include "BlueprintPlayMontageAndWaitTaskProxy.h"
+#include "CompilerResultsLog.h"
+
+#define LOCTEXT_NAMESPACE "K2Node_PlayMontageAndWait"
+
+UK2Node_PlayMontageAndWait::UK2Node_PlayMontageAndWait(const class FPostConstructInitializeProperties& PCIP)
+	: Super(PCIP)
+{
+	
+	ProxyFactoryFunctionName = GET_FUNCTION_NAME_CHECKED(USkillSystemBlueprintLibrary, CreatePlayMontageAndWaitProxy);
+	ProxyFactoryClass = USkillSystemBlueprintLibrary::StaticClass();
+	ProxyClass = UBlueprintPlayMontageAndWaitTaskProxy::StaticClass();
+}
+
+FString UK2Node_PlayMontageAndWait::GetCategoryName()
+{
+	return TEXT("Ability");
+}
+
+FString UK2Node_PlayMontageAndWait::GetTooltip() const
+{
+	return TEXT("Play Anim Montage until it completes or is interrupted");
+}
+
+FText UK2Node_PlayMontageAndWait::GetNodeTitle(ENodeTitleType::Type TitleType) const
+{
+	return LOCTEXT("PlayMontageAndWait", "Play Montage and Wait");
+}
+
+#undef LOCTEXT_NAMESPACE
+
+
diff --git a/Engine/Source/Editor/SkillSystemEditor/Private/SkillSystemEditorModulePrivatePCH.h b/Engine/Source/Editor/SkillSystemEditor/Private/SkillSystemEditorModulePrivatePCH.h
--- a/Engine/Source/Editor/SkillSystemEditor/Private/SkillSystemEditorModulePrivatePCH.h
+++ b/Engine/Source/Editor/SkillSystemEditor/Private/SkillSystemEditorModulePrivatePCH.h
@@ -14,4 +14,4 @@
 #include "Engine.h"
 #include "BlueprintGraphDefinitions.h"
 
-//#include "SkillSystemEditorClasses.h"
\ No newline at end of file
+#include "SkillSystemEditorClasses.h"
\ No newline at end of file
diff --git a/Engine/Source/Runtime/SkillSystem/Classes/AttributeComponent.h b/Engine/Source/Runtime/SkillSystem/Classes/AttributeComponent.h
--- a/Engine/Source/Runtime/SkillSystem/Classes/AttributeComponent.h
+++ b/Engine/Source/Runtime/SkillSystem/Classes/AttributeComponent.h
@@ -28,6 +28,8 @@ class SKILLSYSTEM_API UAttributeComponent : public UActorComponent
 
 	virtual ~UAttributeComponent();
 
+	virtual void InitializeComponent() OVERRIDE;
+
 	template <class T >
 	T*	GetSet()
 	{
@@ -69,16 +71,22 @@ class SKILLSYSTEM_API UAttributeComponent : public UActorComponent
 
 	float GetGameplayEffectDuration(FActiveGameplayEffectHandle Handle) const;
 
+	float GetGameplayEffectDuration() const;
+
 	// Not happy with this interface but don't see a better way yet. How should outside code (UI, etc) ask things like 'how much is this gameplay effect modifying my damage by'
 	// (most likely we want to catch this on the backend - when damage is applied we can get a full dump/history of how the number got to where it is. But still we may need polling methods like below (how much would my damage be)
-	UFUNCTION(BlueprintCallable, BlueprintAuthorityOnly, Category = GameplayEffects)
+	UFUNCTION(BlueprintCallable, Category = GameplayEffects)
 	float GetGameplayEffectMagnitude(FActiveGameplayEffectHandle Handle, FGameplayAttribute Attribute) const;
 
-	UFUNCTION(BlueprintCallable, BlueprintAuthorityOnly, Category = GameplayEffects)
+	UFUNCTION(BlueprintCallable, Category = GameplayEffects)
 	float GetGameplayEffectMagnitudeByTag(FActiveGameplayEffectHandle InHandle, FName GameplayTag) const;
 
-	UFUNCTION(BlueprintCallable, BlueprintAuthorityOnly, Category = GameplayEffects)
+	UFUNCTION(BlueprintCallable, Category = GameplayEffects)
 	bool IsGameplayEffectActive(FActiveGameplayEffectHandle InHandle) const;
+
+	// Tags
+	UFUNCTION(BlueprintCallable, Category=GameplayEffects)
+	bool HasAnyTags(FGameplayTagContainer &Tags);
 	
 
 	// --------------------------------------------
@@ -88,14 +96,25 @@ class SKILLSYSTEM_API UAttributeComponent : public UActorComponent
 	FActiveGameplayEffectHandle ApplyGameplayEffectSpecToTarget(UGameplayEffect *GameplayEffect, UAttributeComponent *Target, float Level = FGameplayEffectLevelSpec::INVALID_LEVEL, FModifierQualifier BaseQualifier = FModifierQualifier()) const;
 
 	UFUNCTION(BlueprintCallable, BlueprintAuthorityOnly, Category = GameplayEffects, meta=(FriendlyName = "ApplyGameplayEffectToSelf"))
-	FActiveGameplayEffectHandle K2_ApplyGameplayEffectToSelf(UGameplayEffect *GameplayEffect, float Level, AActor *Instigator);
+	FActiveGameplayEffectHandle K2_ApplyGameplayEffectToSelf(const UGameplayEffect *GameplayEffect, float Level, AActor *Instigator);
 	
-	FActiveGameplayEffectHandle ApplyGameplayEffectToSelf(UGameplayEffect *GameplayEffect, float Level, AActor *Instigator, FModifierQualifier BaseQualifier = FModifierQualifier() );
+	FActiveGameplayEffectHandle ApplyGameplayEffectToSelf(const UGameplayEffect *GameplayEffect, float Level, AActor *Instigator, FModifierQualifier BaseQualifier = FModifierQualifier() );
 
 	int32 GetNumActiveGameplayEffect() const;
 
 	void AddDependancyToAttribute(FGameplayAttribute Attribute, const TWeakPtr<FAggregator> InDependant);
 
+	/** Tests if all modifiers in this GameplayEffect will leave the attribute > 0.f */
+	bool CanApplyAttributeModifiers(const UGameplayEffect *GameplayEffect, float Level, AActor *Instigator);
+
+	// Generic 'Get expected mangitude (list) if I was to apply this outoging or incoming'
+
+	// Get duration or magnitude (list) of active effects
+	//		-Get duration of CD
+	//		-Get magnitude + duration of a movespeed buff
+
+	TArray<float> GetActiveEffectsTimeRemaining(const FActiveGameplayEffectQuery Query) const;
+
 	// --------------------------------------------
 	// Temp / Debug
 	// --------------------------------------------
@@ -117,6 +136,55 @@ class SKILLSYSTEM_API UAttributeComponent : public UActorComponent
 		}
 	}
 
+	// ----------------------------------------------------------------------------------------------------------------
+	//
+	//	GameplayAbilities
+	//	
+	//	Should this be its own component? What should this do?
+	//		-Load/initialize instanced subobjects
+	//		-Bind to input component (maybe? where would weapons fit in?)
+	//		
+	//
+	// ----------------------------------------------------------------------------------------------------------------
+
+	/** Creates instances of the abilities that we need but doesn't necessarily  */
+	void InitializeAbilities(class UGameplayAbilitySet *Set);
+
+	/** Binds the InputComponent directly to the abilities. This is a convenience for some games. Other games may have abilities on weapons or who need some game logic to decide what ability a button press should go to. */
+	void BindInputComponentToAbilities(UInputComponent *InputComponent, UGameplayAbilitySet *AbilitySet);
+
+	UPROPERTY()
+	TArray<class UGameplayAbility *>	InstancedAbilities;
+
+	/** All abilities, including those that are shared/non instanced. */
+	UPROPERTY(BlueprintReadOnly, Category="Abilities")
+	TArray<class UGameplayAbility *>	AllAbilities;
+
+	/** This is probably temp for testing. We want to think a bit more on the API for outside stuff activating abilities */
+	UFUNCTION(BlueprintCallable, Category="Abilities")
+	bool	ActivateAbility(int32 AbilityIdx);
+
+	// There needs to be a concept of an animating ability. Only one may exist at a time. New requestrs can be queued up, overridden, or ignored.
+
+	UPROPERTY()
+	class UGameplayAbility *	AnimatingAbility;
+
+	void	MontageBranchPoint_AbilityDecisionStop();
+
+	void	MontageBranchPoint_AbilityDecisionStart();
+
+	// --- Ability RPCs
+
+	UFUNCTION(Server, WithValidation, Reliable)
+	void	ServerTryActivateAbility(class UGameplayAbility * AbilityToActivate);
+
+	UFUNCTION(Client, Reliable)
+	void	ClientActivateAbilityFailed(class UGameplayAbility * AbilityToActivate);
+
+	UFUNCTION(Client, Reliable)
+	void	ClientActivateAbilitySucceed(class UGameplayAbility * AbilityToActivate);
+	
+
 private:
 
 	// --------------------------------------------
diff --git a/Engine/Source/Runtime/SkillSystem/Classes/BlueprintPlayMontageAndWaitTaskProxy.h b/Engine/Source/Runtime/SkillSystem/Classes/BlueprintPlayMontageAndWaitTaskProxy.h
new file mode 100644
--- /dev/null
+++ b/Engine/Source/Runtime/SkillSystem/Classes/BlueprintPlayMontageAndWaitTaskProxy.h
@@ -0,0 +1,23 @@
+// Copyright 1998-2014 Epic Games, Inc. All Rights Reserved.
+#pragma once
+#include "BlueprintPlayMontageAndWaitTaskProxy.generated.h"
+
+DECLARE_DYNAMIC_MULTICAST_DELEGATE(FMontageWaitSimpleDelegate);
+
+UCLASS(MinimalAPI)
+class UBlueprintPlayMontageAndWaitTaskProxy : public UObject
+{
+	GENERATED_UCLASS_BODY()
+
+	UPROPERTY(BlueprintAssignable)
+	FMontageWaitSimpleDelegate	OnComplete;
+
+	UPROPERTY(BlueprintAssignable)
+	FMontageWaitSimpleDelegate	OnInterrupted;
+
+	UFUNCTION()
+	void	OnMontageEnded(UAnimMontage* Montage, bool bInterrupted);
+
+public:
+	
+};
\ No newline at end of file
diff --git a/Engine/Source/Runtime/SkillSystem/Classes/GameplayAbility.h b/Engine/Source/Runtime/SkillSystem/Classes/GameplayAbility.h
new file mode 100644
--- /dev/null
+++ b/Engine/Source/Runtime/SkillSystem/Classes/GameplayAbility.h
@@ -0,0 +1,209 @@
+// Copyright 1998-2013 Epic Games, Inc. All Rights Reserved.
+
+#pragma once
+
+#include "Runtime/Engine/Classes/Animation/AnimInstance.h"
+#include "GameplayAbility.generated.h"
+
+class UAnimInstance;
+class UAttributeComponent;
+
+UENUM(BlueprintType)
+namespace EGameplayAbilityReplication
+{
+	enum Type
+	{
+		// Numeric
+		Predictive		UMETA(DisplayName = "Predictive"),	// Part of this ability runs predictively on the client.
+		Server			UMETA(DisplayName = "Server"),		// This ability must be OK'd by the server before doing anything on a client.
+		Client			UMETA(DisplayName = "Client"),		// This ability runs as long the client says it does.
+	};
+}
+
+USTRUCT()
+struct SKILLSYSTEM_API FGameplayAbilityActorInfo
+{
+	GENERATED_USTRUCT_BODY()
+
+	TWeakObjectPtr<AActor>	Actor;
+
+	TWeakObjectPtr<UAnimInstance>	AnimInstance;
+
+	TWeakObjectPtr<UAttributeComponent>	AttributeComponent;
+
+	void InitFromActor(AActor *Actor);
+};
+
+USTRUCT(BlueprintType)
+struct FGameplayAbilityHandle
+{
+	GENERATED_USTRUCT_BODY()
+
+	FGameplayAbilityHandle()
+	: Handle(INDEX_NONE)
+	{
+
+	}
+
+	FGameplayAbilityHandle(int32 InHandle)
+		: Handle(InHandle)
+	{
+
+	}
+
+	bool IsValid() const
+	{
+		return Handle != INDEX_NONE;
+	}
+
+	FGameplayAbilityHandle GetNextHandle()
+	{
+		return FGameplayAbilityHandle(Handle + 1);
+	}
+
+	bool operator==(const FGameplayAbilityHandle& Other) const
+	{
+		return Handle == Other.Handle;
+	}
+
+	bool operator!=(const FGameplayAbilityHandle& Other) const
+	{
+		return Handle != Other.Handle;
+	}
+
+	FString ToString() const
+	{
+		return FString::Printf(TEXT("%d"), Handle);
+	}
+
+private:
+
+	UPROPERTY()
+	int32 Handle;
+};
+
+USTRUCT()
+struct FGameplayAbilitySpec
+{
+	GENERATED_USTRUCT_BODY()
+
+	UPROPERTY()
+	int32	Level;
+
+	UPROPERTY()
+	class UGameplayAbility * Ability;
+
+	UPROPERTY()
+	FGameplayAbilityHandle	Handle;
+};
+
+/**
+* UGameplayEffect
+*	The GameplayEffect definition. This is the data asset defined in the editor that drives everything.
+*/
+UCLASS()
+class SKILLSYSTEM_API UGameplayAbility : public UDataAsset
+{
+	GENERATED_UCLASS_BODY()
+
+public:
+	
+	// ----------------------------------------------------------------------------------------------------------------
+	//
+	//	Primary events that are implemented to define an ability
+	//	
+	// ----------------------------------------------------------------------------------------------------------------
+	
+	/** Returns true if this ability can be activated right now. Has no side effects! */
+	virtual bool CanActivateAbility(const FGameplayAbilityActorInfo ActorInfo) const;
+	
+	/** Attempts to commit the ability (spend resources, etc). This our last chance to fail. */
+	virtual bool CommitAbility(const FGameplayAbilityActorInfo ActorInfo);
+
+	/** The last chance to fail before commiting */
+	virtual bool CommitCheck(const FGameplayAbilityActorInfo ActorInfo);
+	
+	/** Does the commmit atomically (consume resources, do cooldowns, etc) */
+	virtual void CommitExecute(const FGameplayAbilityActorInfo ActorInfo);
+
+	// -----------------------------------------------
+
+	/** Calls ActivateAbility if we CanACtivateAbility */
+	virtual bool TryActivateAbility(const FGameplayAbilityActorInfo ActorInfo);
+
+	/** (Optionally Does Stuff) -> Commit Ability -> (Optionally Does Stuff) */
+	virtual void ActivateAbility(const FGameplayAbilityActorInfo ActorInfo);
+
+	/** Input binding. Base implementation calls TryActivateAbility */
+	virtual void InputPressed(int32 InputID, const FGameplayAbilityActorInfo ActorInfo);
+
+	/** Input binding. Base implementation does nothing */
+	virtual void InputReleased(int32 InputID, const FGameplayAbilityActorInfo ActorInfo);
+
+	// ----------------------------------------------------------------------------------------------------------------
+	//
+	//	Outward facing API for other systems
+	//
+	// ----------------------------------------------------------------------------------------------------------------
+
+	float GetCooldownTimeRemaining(const FGameplayAbilityActorInfo ActorInfo) const;
+
+	// ----------------------------------------------------------------------------------------------------------------
+	//
+	//	Standardized State
+	//
+	// ----------------------------------------------------------------------------------------------------------------	
+
+	UPROPERTY(EditDefaultsOnly, Category=Replication)
+	TEnumAsByte<EGameplayAbilityReplication::Type>	ReplicationType;
+
+	/** This GameplayEffect represents the cooldown. It will be applied when the ability is commited and the ability cannot be used again until it is expired. */
+	UPROPERTY(EditDefaultsOnly, Category=Cooldowns)
+	class UGameplayEffect * CooldownGameplayEffect;
+
+	/** This GameplayEffect represents the cost (mana, stamina, etc) of the ability. It will be applied when the ability is commited. */
+	UPROPERTY(EditDefaultsOnly, Category=Costs)
+	class UGameplayEffect * CostGameplayEffect;
+
+	// ----------------------------------------------------------------------------------------------------------------
+	//
+	//	Animation callbacks
+	//
+	// ----------------------------------------------------------------------------------------------------------------
+
+	void MontageBranchPoint_AbilityDecisionStop(const FGameplayAbilityActorInfo ActorInfo) const;
+
+	void MontageBranchPoint_AbilityDecisionStart(const FGameplayAbilityActorInfo ActorInfo) const;
+
+	// ----------------------------------------------------------------------------------------------------------------
+
+	virtual void PredictiveActivateAbility(const FGameplayAbilityActorInfo ActorInfo);
+
+	virtual void ServerTryActivateAbility(const FGameplayAbilityActorInfo ActorInfo);
+
+	virtual void ClientActivateAbility(const FGameplayAbilityActorInfo ActorInfo);
+
+	// ----------------------------------------------------------------------------------------------------------------
+
+	virtual bool AllowInstancing() const
+	{
+		return false;
+	}
+
+	virtual UWorld* GetWorld() const OVERRIDE
+	{
+		return GetOuter()->GetWorld();
+	}
+
+private:
+
+	/** Checks cooldown. returns true if we can be used again. False if not */
+	bool	CheckCooldown(const FGameplayAbilityActorInfo ActorInfo) const;
+
+	/** Applies CooldownGameplayEffect to the target */
+	void	ApplyCooldown(const FGameplayAbilityActorInfo ActorInfo);
+
+	bool	CheckCost(const FGameplayAbilityActorInfo ActorInfo) const;
+
+	void	ApplyCost(const FGameplayAbilityActorInfo ActorInfo);
+};
\ No newline at end of file
diff --git a/Engine/Source/Runtime/SkillSystem/Classes/GameplayAbilitySet.h b/Engine/Source/Runtime/SkillSystem/Classes/GameplayAbilitySet.h
new file mode 100644
--- /dev/null
+++ b/Engine/Source/Runtime/SkillSystem/Classes/GameplayAbilitySet.h
@@ -0,0 +1,45 @@
+// Copyright 1998-2013 Epic Games, Inc. All Rights Reserved.
+
+#pragma once
+
+#include "GameplayAbilitySet.generated.h"
+
+USTRUCT()
+struct FGameplayAbilityBindInfoCommandIDPair
+{
+	GENERATED_USTRUCT_BODY()
+
+	// We want this to be a drop down of valid lists but needs to be defined per game
+	UPROPERTY(EditAnywhere, Category = BindInfo)
+	FString		CommandString;
+
+	UPROPERTY(EditAnywhere, Category = BindInfo)
+	int32		InputID;
+};
+
+USTRUCT()
+struct FGameplayAbilityBindInfo
+{
+	GENERATED_USTRUCT_BODY()
+
+	UPROPERTY(EditAnywhere, Category=BindInfo)
+	TArray<FGameplayAbilityBindInfoCommandIDPair>	CommandBindings;
+
+	UPROPERTY(EditAnywhere, Category=BindInfo)
+	TSubclassOf<UGameplayAbility>	GameplayAbilityClass;
+};
+
+/**
+* UGameplayAbilitySet
+*	This contains a list of abilities along with key bindings. THis will be very game specific, so it is expected for games to override this clas.
+*/
+UCLASS(Blueprintable)
+class SKILLSYSTEM_API UGameplayAbilitySet : public UDataAsset
+{
+	GENERATED_UCLASS_BODY()
+
+public:
+
+	UPROPERTY(EditAnywhere, Category=AbilitySet)
+	TArray<FGameplayAbilityBindInfo>	Abilities;
+};
\ No newline at end of file
diff --git a/Engine/Source/Runtime/SkillSystem/Classes/GameplayAbility_Instanced.h b/Engine/Source/Runtime/SkillSystem/Classes/GameplayAbility_Instanced.h
new file mode 100644
--- /dev/null
+++ b/Engine/Source/Runtime/SkillSystem/Classes/GameplayAbility_Instanced.h
@@ -0,0 +1,51 @@
+// Copyright 1998-2013 Epic Games, Inc. All Rights Reserved.
+
+#pragma once
+
+#include "GameplayAbility_Instanced.generated.h"
+
+class UAnimInstance;
+class UAttributeComponent;
+
+/**
+* UGameplayEffect
+*	The GameplayEffect definition. This is the data asset defined in the editor that drives everything.
+*/
+UCLASS(Blueprintable)
+class SKILLSYSTEM_API UGameplayAbility_Instanced : public UGameplayAbility
+{
+	GENERATED_UCLASS_BODY()
+
+public:
+	
+	// -------------------------------------------------
+
+	UFUNCTION(BlueprintImplementableEvent, Category = Ability)
+	virtual bool K2_CanActivateAbility(const FGameplayAbilityActorInfo ActorInfo) const;
+
+	virtual bool CanActivateAbility(const FGameplayAbilityActorInfo ActorInfo) const OVERRIDE;	
+
+	// -------------------------------------------------
+
+	UFUNCTION(BlueprintImplementableEvent, Category = Ability)
+	virtual void K2_CommitExecute(const FGameplayAbilityActorInfo ActorInfo);
+
+	virtual void CommitExecute(const FGameplayAbilityActorInfo ActorInfo) OVERRIDE;
+
+	// -------------------------------------------------
+
+	UFUNCTION(BlueprintImplementableEvent, Category = Ability)
+	virtual void K2_ActivateAbility(const FGameplayAbilityActorInfo ActorInfo);
+
+	virtual void ActivateAbility(const FGameplayAbilityActorInfo ActorInfo) OVERRIDE;
+
+	// -------------------------------------------------
+
+	virtual bool AllowInstancing() const OVERRIDE
+	{
+		return true;
+	}
+
+private:
+	
+};
\ No newline at end of file
diff --git a/Engine/Source/Runtime/SkillSystem/Classes/GameplayAbility_Montage.h b/Engine/Source/Runtime/SkillSystem/Classes/GameplayAbility_Montage.h
new file mode 100644
--- /dev/null
+++ b/Engine/Source/Runtime/SkillSystem/Classes/GameplayAbility_Montage.h
@@ -0,0 +1,33 @@
+// Copyright 1998-2013 Epic Games, Inc. All Rights Reserved.
+
+#pragma once
+
+#include "GameplayAbility_Montage.generated.h"
+
+/**
+ *	A gameplay ability that plays a single montage and applies a GameplayEffect
+ */
+UCLASS()
+class SKILLSYSTEM_API UGameplayAbility_Montage : public UGameplayAbility
+{
+	GENERATED_UCLASS_BODY()
+
+public:
+	
+	virtual void ActivateAbility(const FGameplayAbilityActorInfo OwnerInfo) OVERRIDE;
+
+	UPROPERTY(EditDefaultsOnly, Category = MontageAbility)
+	UAnimMontage *	MontageToPlay;
+
+	UPROPERTY(EditDefaultsOnly, Category = MontageAbility)
+	float	PlayRate;
+
+	UPROPERTY(EditDefaultsOnly, Category = MontageAbility)
+	FName	SectionName;
+
+	/** GameplayEffects to apply and then remove while the animation is playing */
+	UPROPERTY(EditDefaultsOnly, Category = MontageAbility)
+	TArray<const UGameplayEffect*>	GameplayEffectsWhileAnimating;
+
+	void OnMontageEnded(UAnimMontage* Montage, bool bInterrupted, const FGameplayAbilityActorInfo OwnerInfo, TArray<struct FActiveGameplayEffectHandle>	AppliedEffects);
+};
\ No newline at end of file
diff --git a/Engine/Source/Runtime/SkillSystem/Classes/GameplayEffect.h b/Engine/Source/Runtime/SkillSystem/Classes/GameplayEffect.h
--- a/Engine/Source/Runtime/SkillSystem/Classes/GameplayEffect.h
+++ b/Engine/Source/Runtime/SkillSystem/Classes/GameplayEffect.h
@@ -611,8 +611,11 @@ struct FGameplayEffectInstigatorContext
  *		the buff expires on the source, the applied DOT is still buffed.
  *
  */
+USTRUCT()
 struct FAggregatorRef
 {
+	GENERATED_USTRUCT_BODY()
+
 	friend struct FAggregatorRef;
 
 	FAggregatorRef()
@@ -668,15 +671,29 @@ struct FAggregatorRef
 	/** Become a hard reference to a new copy of what we are reference AND make new copies/hard refs of the complete modifier chain in our FAggregator */
 	void MakeUniqueDeep();
 
+	bool NetSerialize(FArchive& Ar, class UPackageMap* Map, bool& bOutSuccess);	
+
 	FString ToString() const;
 	void PrintAll() const;
 
+	UPROPERTY()
+	int32	Foo;
+
 private:
 
 	TSharedPtr<struct FAggregator>	SharedPtr;
 	TWeakPtr<struct FAggregator> WeakPtr;
 };
 
+template<>
+struct TStructOpsTypeTraits< FAggregatorRef > : public TStructOpsTypeTraitsBase
+{
+	enum
+	{
+		WithNetSerializer = true
+	};
+};
+
 /**
  * GameplayModifierEvaluatedData
  *	This is the data that FAggregator aggregates and turns into FGameplayModifierEvaluatedData.
@@ -832,6 +849,8 @@ struct FAggregator : public TSharedFromThis<FAggregator>
 	void RefreshDependencies();
 	void MakeUniqueDeep();
 
+	void SetFromNetSerialize(float NetSerialize);
+
 	// ----------------------------------------------------------------------------
 	// This is data only used in debugging/tracking where aggregator's came from
 	// ----------------------------------------------------------------------------
@@ -1069,7 +1088,9 @@ struct FGameplayEffectSpec
 	float GetPeriod() const;
 	float GetMagnitude(const FGameplayAttribute &Attribute) const;
 
+	UPROPERTY()
 	FAggregatorRef	Duration;
+
 	FAggregatorRef	Period;
 
 	EGameplayEffectStackingPolicy::Type StackingPolicy;
@@ -1116,22 +1137,24 @@ struct FActiveGameplayEffect : public FFastArraySerializerItem
 	GENERATED_USTRUCT_BODY()
 
 	FActiveGameplayEffect()
-		: StartTime(0.f)
+		: StartGameStateTime(0)
+		, StartWorldTime(0.f)
 		, NextExecuteTime(0.f)
 	{
 	}
 
-	FActiveGameplayEffect(FActiveGameplayEffectHandle InHandle, const FGameplayEffectSpec &InSpec, float CurrentTime = 0.f)
+	FActiveGameplayEffect(FActiveGameplayEffectHandle InHandle, const FGameplayEffectSpec &InSpec, float CurrentWorldTime, int32 InStartGameStateTime)
 		: Handle(InHandle)
 		, Spec(InSpec)
-		, StartTime(CurrentTime)
+		, StartGameStateTime(InStartGameStateTime)
+		, StartWorldTime(CurrentWorldTime)
 		, NextExecuteTime(0.f)
 	{
 		// Init NextExecuteTime if necessary
 		float Period = GetPeriod();
 		if (Period != UGameplayEffect::NO_PERIOD)
 		{
-			NextExecuteTime = CurrentTime + SMALL_NUMBER;
+			NextExecuteTime = CurrentWorldTime + SMALL_NUMBER;
 		}
 
 		for (FModifierSpec &Mod : Spec.Modifiers)
@@ -1147,7 +1170,10 @@ struct FActiveGameplayEffect : public FFastArraySerializerItem
 
 	/** Game time this started */
 	UPROPERTY()
-	float StartTime;
+	int32 StartGameStateTime;
+
+	UPROPERTY(NotReplicated)
+	float StartWorldTime;
 
 	UPROPERTY(NotReplicated)
 	float NextExecuteTime;
@@ -1179,6 +1205,58 @@ struct FActiveGameplayEffect : public FFastArraySerializerItem
 	}
 };
 
+
+USTRUCT()
+struct FActiveGameplayEffectData
+{
+	GENERATED_USTRUCT_BODY()
+
+	FActiveGameplayEffectData()
+		: Handle()
+		, Duration(0.f)
+		, Magnitude(0.f)
+	{
+	}
+
+	FActiveGameplayEffectData(const FActiveGameplayEffectHandle InHandle, const float InDuration, const float InMagnitude)
+		: Handle(InHandle)
+		, Duration(InDuration)
+		, Magnitude(InMagnitude)
+	{
+	}
+
+	UPROPERTY()
+	FActiveGameplayEffectHandle	Handle;
+
+	UPROPERTY()
+	float	Duration;
+
+	UPROPERTY()
+	float	Magnitude;
+};
+
+/** Generic querying data structure for active GameplayEffects. Lets us ask things like:
+ *		-Give me duration/magnitude of active gameplay effects with these tags
+ *		-
+ */
+USTRUCT()
+struct FActiveGameplayEffectQuery
+{
+	GENERATED_USTRUCT_BODY()
+
+	FActiveGameplayEffectQuery()
+		: TagContainer(NULL)
+	{
+	}
+
+	FActiveGameplayEffectQuery(const FGameplayTagContainer * InTagContainer)
+		: TagContainer(InTagContainer)
+	{
+	}
+
+	const FGameplayTagContainer *	TagContainer;
+};
+
 /**
  * Active GameplayEffects Container
  *	-Bucket of ActiveGameplayEffects
@@ -1197,7 +1275,7 @@ struct FActiveGameplayEffectsContainer : public FFastArraySerializer
 	UPROPERTY()
 	TArray< FActiveGameplayEffect >	GameplayEffects;
 	
-	FActiveGameplayEffect & CreateNewActiveGameplayEffect(const FGameplayEffectSpec &Spec, float GameTimeSeconds);
+	FActiveGameplayEffect & CreateNewActiveGameplayEffect(const FGameplayEffectSpec &Spec);
 
 	void ApplyActiveEffectsTo(OUT FGameplayEffectSpec &Spec, const FModifierQualifier &QualifierContext) const;
 
@@ -1246,6 +1324,15 @@ struct FActiveGameplayEffectsContainer : public FFastArraySerializer
 
 	bool bNeedToRecalculateStacks;
 
+	bool HasAnyTags(FGameplayTagContainer &Tags);
+
+	bool CanApplyAttributeModifiers(const UGameplayEffect *GameplayEffect, float Level, AActor *Instigator);
+	
+	TArray<float> GetActiveEffectsTimeRemaining(const FActiveGameplayEffectQuery Query) const;
+
+	int32 GetGameStateTime() const;
+	float GetWorldTime() const;
+
 private:
 
 	FAggregatorRef & FindOrCreateAttributeAggregator(FGameplayAttribute Attribute);
diff --git a/Engine/Source/Runtime/SkillSystem/Classes/SkillSystemBlueprintLibrary.h b/Engine/Source/Runtime/SkillSystem/Classes/SkillSystemBlueprintLibrary.h
new file mode 100644
--- /dev/null
+++ b/Engine/Source/Runtime/SkillSystem/Classes/SkillSystemBlueprintLibrary.h
@@ -0,0 +1,21 @@
+// Copyright 1998-2014 Epic Games, Inc. All Rights Reserved.
+
+#pragma once
+
+#include "GameplayAbility.h"
+#include "SkillSystemBlueprintLibrary.generated.h"
+
+class UBlueprintPlayMontageAndWaitTaskProxy;
+
+UCLASS(MinimalAPI)
+class USkillSystemBlueprintLibrary : public UBlueprintFunctionLibrary
+{
+	GENERATED_UCLASS_BODY()
+
+	UFUNCTION(BlueprintCallable, Category = "Utilities|FlowControl", meta = (Latent, HidePin = "WorldContextObject", DefaultToSelf = "WorldContextObject", LatentInfo = "LatentInfo" ))
+	static void	StaticFuncPlayMontageAndWait(AActor* WorldContextObject, UAnimMontage * Montage, struct FLatentActionInfo LatentInfo);
+
+
+	UFUNCTION(BlueprintCallable, meta=(HidePin="WorldContextObject", DefaultToSelf="WorldContextObject", BlueprintInternalUseOnly = "TRUE"))
+	static UBlueprintPlayMontageAndWaitTaskProxy* CreatePlayMontageAndWaitProxy(class UObject* WorldContextObject, class UAnimMontage *MontageToPlay);
+};
\ No newline at end of file
diff --git a/Engine/Source/Runtime/SkillSystem/Classes/SkillSystemTestAttributeSet.h b/Engine/Source/Runtime/SkillSystem/Classes/SkillSystemTestAttributeSet.h
--- a/Engine/Source/Runtime/SkillSystem/Classes/SkillSystemTestAttributeSet.h
+++ b/Engine/Source/Runtime/SkillSystem/Classes/SkillSystemTestAttributeSet.h
@@ -1,47 +1,54 @@
 // Copyright 1998-2014 Epic Games, Inc. All Rights Reserved.
 
 #pragma once
+#include "AttributeSet.h"
 #include "SkillSystemTestAttributeSet.generated.h"
 
 UCLASS(Blueprintable, BlueprintType)
-class USkillSystemTestAttributeSet : public UAttributeSet
+class SKILLSYSTEM_API USkillSystemTestAttributeSet : public UAttributeSet
 {
 	GENERATED_UCLASS_BODY()
 
-	UPROPERTY(EditAnywhere, BlueprintReadWrite, Category = "AttributeTest", meta = (HideFromModifiers))			// You can't make a GameplayEffect modify Health Directly (go through)
+	UPROPERTY(EditAnywhere, BlueprintReadWrite, Replicated, Category = "AttributeTest", meta = (HideFromModifiers))			// You can't make a GameplayEffect modify Health Directly (go through)
 	float	MaxHealth;
 
-	UPROPERTY(EditAnywhere, BlueprintReadWrite, Category = "AttributeTest", meta = (HideFromModifiers))			// You can't make a GameplayEffect modify Health Directly (go through)
+	UPROPERTY(EditAnywhere, BlueprintReadWrite, Replicated, Category = "AttributeTest", meta = (HideFromModifiers))			// You can't make a GameplayEffect modify Health Directly (go through)
 	float	Health;
 
+	UPROPERTY(EditAnywhere, BlueprintReadWrite, Replicated, Category = "AttributeTest")
+	float	Mana;
+
+	UPROPERTY(EditAnywhere, BlueprintReadWrite, Replicated, Category = "AttributeTest")
+	float	MaxMana;
+
 	/** This Damage is just used for applying negative health mods. Its not a 'persistent' attribute. */
 	UPROPERTY(EditAnywhere, BlueprintReadWrite, Category = "AttributeTest", meta = (HideFromLevelInfos))		// You can't make a GameplayEffect 'powered' by Damage (Its transient)
 	float	Damage;
 
 	/** This Attribute is the actual spell damage for this actor. It will power spell based GameplayEffects */
-	UPROPERTY(EditAnywhere, BlueprintReadWrite, Category = "AttributeTest")
+	UPROPERTY(EditAnywhere, BlueprintReadWrite, Replicated, Category = "AttributeTest")
 	float	SpellDamage;
 
 	/** This Attribute is the actual physical damage for this actor. It will power physical based GameplayEffects */
-	UPROPERTY(EditAnywhere, BlueprintReadWrite, Category = "AttributeTest")
+	UPROPERTY(EditAnywhere, BlueprintReadWrite, Replicated, Category = "AttributeTest")
 	float	PhysicalDamage;
 
-	UPROPERTY(EditAnywhere, BlueprintReadWrite, Category = "AttributeTest")
+	UPROPERTY(EditAnywhere, BlueprintReadWrite, Replicated, Category = "AttributeTest")
 	float	CritChance;
 
-	UPROPERTY(EditAnywhere, BlueprintReadWrite, Category = "AttributeTest")
+	UPROPERTY(EditAnywhere, BlueprintReadWrite, Replicated, Category = "AttributeTest")
 	float	CritMultiplier;
 
-	UPROPERTY(EditAnywhere, BlueprintReadWrite, Category = "AttributeTest")
+	UPROPERTY(EditAnywhere, BlueprintReadWrite, Replicated, Category = "AttributeTest")
 	float	ArmorDamageReduction;
 
-	UPROPERTY(EditAnywhere, BlueprintReadWrite, Category = "AttributeTest")
+	UPROPERTY(EditAnywhere, BlueprintReadWrite, Replicated, Category = "AttributeTest")
 	float	DodgeChance;
 
-	UPROPERTY(EditAnywhere, BlueprintReadWrite, Category = "AttributeTest")
+	UPROPERTY(EditAnywhere, BlueprintReadWrite, Replicated, Category = "AttributeTest")
 	float	LifeSteal;
 
-	UPROPERTY(EditAnywhere, BlueprintReadWrite, Category = "AttributeTest")
+	UPROPERTY(EditAnywhere, BlueprintReadWrite, Replicated, Category = "AttributeTest")
 	float	Strength;
 
 	UPROPERTY(EditAnywhere, BlueprintReadWrite, Category = "AttributeTest")
@@ -55,4 +62,4 @@ class USkillSystemTestAttributeSet : public UAttributeSet
 
 	virtual void PreAttributeModify(struct FGameplayEffectModCallbackData &Data) OVERRIDE;
 	virtual void PostAttributeModify(const struct FGameplayEffectModCallbackData &Data) OVERRIDE;
-};
\ No newline at end of file
+};
diff --git a/Engine/Source/Runtime/SkillSystem/Classes/SkillSystemTestPawn.h b/Engine/Source/Runtime/SkillSystem/Classes/SkillSystemTestPawn.h
--- a/Engine/Source/Runtime/SkillSystem/Classes/SkillSystemTestPawn.h
+++ b/Engine/Source/Runtime/SkillSystem/Classes/SkillSystemTestPawn.h
@@ -30,5 +30,8 @@ class SKILLSYSTEM_API ASkillSystemTestPawn : public ADefaultPawn, public IGamepl
 	UPROPERTY(EditDefaultsOnly, Category=GameplayEffects)
 	FGameplayCueHandler	GameplayCueHandler;
 
+	//UPROPERTY(EditDefaultsOnly, Category=GameplayEffects)
+	//UGameplayAbilitySet * DefaultAbilitySet;
+
 	static FName AttributeComponentName;
 };
\ No newline at end of file
diff --git a/Engine/Source/Runtime/SkillSystem/Private/AttributeComponent.cpp b/Engine/Source/Runtime/SkillSystem/Private/AttributeComponent.cpp
--- a/Engine/Source/Runtime/SkillSystem/Private/AttributeComponent.cpp
+++ b/Engine/Source/Runtime/SkillSystem/Private/AttributeComponent.cpp
@@ -18,11 +18,15 @@ DEFINE_LOG_CATEGORY(LogAttributeComponent);
 UAttributeComponent::UAttributeComponent(const class FPostConstructInitializeProperties& PCIP)
 	: Super(PCIP)
 {
+	bWantsInitializeComponent = true;
+
 	PrimaryComponentTick.TickGroup = TG_DuringPhysics;
 	PrimaryComponentTick.bStartWithTickEnabled = true; // FIXME! Just temp until timer manager figured out
 	PrimaryComponentTick.bCanEverTick = true;
 	
 	ActiveGameplayEffects.Owner = this;
+
+	SetIsReplicated(true);
 }
 
 UAttributeComponent::~UAttributeComponent()
@@ -177,13 +181,13 @@ FActiveGameplayEffectHandle UAttributeComponent::ApplyGameplayEffectSpecToTarget
 }
 
 /** Helper function since we can't have default/optional values for FModifierQualifier in K2 function */
-FActiveGameplayEffectHandle UAttributeComponent::K2_ApplyGameplayEffectToSelf(UGameplayEffect *GameplayEffect, float Level, AActor *Instigator)
+FActiveGameplayEffectHandle UAttributeComponent::K2_ApplyGameplayEffectToSelf(const UGameplayEffect *GameplayEffect, float Level, AActor *Instigator)
 {
 	return ApplyGameplayEffectToSelf(GameplayEffect, Level, Instigator);
 }
 
 /** This is a helper function - it seems like this will be useful as a blueprint interface at the least, but Level parameter may need to be expanded */
-FActiveGameplayEffectHandle UAttributeComponent::ApplyGameplayEffectToSelf(UGameplayEffect *GameplayEffect, float Level, AActor *Instigator, FModifierQualifier BaseQualifier)
+FActiveGameplayEffectHandle UAttributeComponent::ApplyGameplayEffectToSelf(const UGameplayEffect *GameplayEffect, float Level, AActor *Instigator, FModifierQualifier BaseQualifier)
 {
 	FGameplayEffectSpec	Spec(GameplayEffect, TSharedPtr<FGameplayEffectLevelSpec>(new FGameplayEffectLevelSpec(Level, GameplayEffect->LevelInfo, GetOwner())), GetCurveDataOverride());
 	Spec.Def = GameplayEffect;
@@ -207,6 +211,12 @@ bool UAttributeComponent::IsGameplayEffectActive(FActiveGameplayEffectHandle InH
 	return ActiveGameplayEffects.IsGameplayEffectActive(InHandle);
 }
 
+bool UAttributeComponent::HasAnyTags(FGameplayTagContainer &Tags)
+{
+	// Fixme: we could aggregate our current tags into a map to avoid this type of iteration
+	return ActiveGameplayEffects.HasAnyTags(Tags);
+}
+
 void UAttributeComponent::TEMP_ApplyActiveGameplayEffects()
 {
 	for (int32 idx=0; idx < ActiveGameplayEffects.GameplayEffects.Num(); ++idx)
@@ -240,14 +250,7 @@ FActiveGameplayEffectHandle UAttributeComponent::ApplyGameplayEffectSpecToSelf(c
 		float Duration = Spec.GetDuration();
 		if (Duration != UGameplayEffect::INSTANT_APPLICATION)
 		{
-			float GameTime = GetWorld()->GetTimeSeconds();
-			// if we have a game state use the ElapsedTime so client server time replication is more accurate else use the time seconds of the game world.
-			if (GetWorld()->GetGameState<AGameState>())
-			{
-				GameTime = GetWorld()->GetGameState<AGameState>()->ElapsedTime;
-			}
-
-			FActiveGameplayEffect &NewActiveEffect = ActiveGameplayEffects.CreateNewActiveGameplayEffect(Spec, GameTime);
+			FActiveGameplayEffect &NewActiveEffect = ActiveGameplayEffects.CreateNewActiveGameplayEffect(Spec);
 			MyHandle = NewActiveEffect.Handle;
 			OurCopyOfSpec = &NewActiveEffect.Spec;
 
@@ -446,6 +449,16 @@ void UAttributeComponent::AddDependancyToAttribute(FGameplayAttribute Attribute,
 	ActiveGameplayEffects.AddDependancyToAttribute(Attribute, InDependant);
 }
 
+bool UAttributeComponent::CanApplyAttributeModifiers(const UGameplayEffect *GameplayEffect, float Level, AActor *Instigator)
+{
+	return ActiveGameplayEffects.CanApplyAttributeModifiers(GameplayEffect, Level, Instigator);
+}
+
+TArray<float> UAttributeComponent::GetActiveEffectsTimeRemaining(const FActiveGameplayEffectQuery Query) const
+{
+	return ActiveGameplayEffects.GetActiveEffectsTimeRemaining(Query);
+}
+
 // ---------------------------------------------------------------------------------------
 
 void UAttributeComponent::PrintAllGameplayEffects() const
@@ -467,7 +480,7 @@ void FActiveGameplayEffectsContainer::PrintAllGameplayEffects() const
 void FActiveGameplayEffect::PrintAll() const
 {
 	SKILL_LOG(Log, TEXT("Handle: %s"), *Handle.ToString());
-	SKILL_LOG(Log, TEXT("StartTime: %.2f"), StartTime);
+	SKILL_LOG(Log, TEXT("StartWorldTime: %.2f"), StartWorldTime);
 	SKILL_LOG(Log, TEXT("NextExecuteTime: %.2f"), NextExecuteTime);
 	Spec.PrintAll();
 }
@@ -598,6 +611,7 @@ void UAttributeComponent::TEMP_TimerTestCallback(int32 x)
 
 void UAttributeComponent::TickComponent(float DeltaTime, enum ELevelTick TickType, FActorComponentTickFunction *ThisTickFunction)
 {
+	SCOPE_CYCLE_COUNTER(STAT_GameplayEffectsTick);
 	SKILL_LOG_SCOPE(TEXT("Ticking %s"), *GetName());
 
 	// This should be tmep until timer manager stuff is figured out
diff --git a/Engine/Source/Runtime/SkillSystem/Private/AttributeComponent_Abilities.cpp b/Engine/Source/Runtime/SkillSystem/Private/AttributeComponent_Abilities.cpp
new file mode 100644
--- /dev/null
+++ b/Engine/Source/Runtime/SkillSystem/Private/AttributeComponent_Abilities.cpp
@@ -0,0 +1,173 @@
+// Copyright 1998-2014 Epic Games, Inc. All Rights Reserved.
+// ActorComponent.cpp: Actor component implementation.
+
+#include "SkillSystemModulePrivatePCH.h"
+
+#include "Net/UnrealNetwork.h"
+#include "MessageLog.h"
+#include "UObjectToken.h"
+#include "MapErrors.h"
+
+#define LOCTEXT_NAMESPACE "AttributeComponent"
+
+/** Enable to log out all render state create, destroy and updatetransform events */
+#define LOG_RENDER_STATE 0
+
+void UAttributeComponent::InitializeComponent()
+{
+	Super::InitializeComponent();
+}
+void UAttributeComponent::InitializeAbilities(class UGameplayAbilitySet *AbilitySet)
+{
+	AActor * OwnerActor = GetOwner();
+	check(OwnerActor);
+
+	FGameplayAbilityActorInfo	ActorInfo;
+	ActorInfo.InitFromActor(OwnerActor);
+
+	for (const FGameplayAbilityBindInfo &BindInfo : AbilitySet->Abilities)
+	{
+		UGameplayAbility * Ability = NULL;
+		if (BindInfo.GameplayAbilityClass)
+		{
+			Ability = BindInfo.GameplayAbilityClass->GetDefaultObject<UGameplayAbility>();
+			if (Ability->AllowInstancing())
+			{
+				Ability = ConstructObject<UGameplayAbility>(BindInfo.GameplayAbilityClass, OwnerActor);
+				InstancedAbilities.Add(Ability);
+			}
+			AllAbilities.Add(Ability);
+		}
+	}
+}
+
+/** This is a convenience/helper function to bind input directly to abilities. Some games may need to do things differently, if for example you want weapons to fire abilities and the pawn to swap weapons. */
+void UAttributeComponent::BindInputComponentToAbilities(UInputComponent *InputComponent, UGameplayAbilitySet *AbilitySet)
+{
+	check(AbilitySet->Abilities.Num() == AllAbilities.Num());
+
+	AActor * OwnerActor = GetOwner();
+	check(OwnerActor);
+
+	FGameplayAbilityActorInfo	ActorInfo;
+	ActorInfo.InitFromActor(OwnerActor);
+
+	for (int32 idx = 0; idx < AbilitySet->Abilities.Num(); ++idx)
+	{
+		FGameplayAbilityBindInfo & BindInfo = AbilitySet->Abilities[idx];
+		UGameplayAbility * Ability = AllAbilities[idx];
+		
+		check(Ability->GetClass() == AbilitySet->Abilities[idx].GameplayAbilityClass);
+		
+		for (const FGameplayAbilityBindInfoCommandIDPair &CommandPair : BindInfo.CommandBindings)
+		{
+			// Pressed event
+			{
+				FInputActionBinding AB(FName(*CommandPair.CommandString), IE_Pressed);
+				AB.ActionDelegate.GetDelegateForManualSet().BindUObject(Ability, &UGameplayAbility::InputPressed, CommandPair.InputID, ActorInfo);
+				InputComponent->AddActionBinding(AB);
+			}
+
+			// Released Event
+			{
+				FInputActionBinding AB(FName(*CommandPair.CommandString), IE_Released);
+				AB.ActionDelegate.GetDelegateForManualSet().BindUObject(Ability, &UGameplayAbility::InputReleased, CommandPair.InputID, ActorInfo);
+				InputComponent->AddActionBinding(AB);
+			}
+		}
+	}
+}
+
+void UAttributeComponent::MontageBranchPoint_AbilityDecisionStop()
+{
+	if (AnimatingAbility)
+	{
+		AActor * OwnerActor = GetOwner();
+		check(OwnerActor);
+
+		FGameplayAbilityActorInfo	ActorInfo;
+		ActorInfo.InitFromActor(OwnerActor);
+
+		AnimatingAbility->MontageBranchPoint_AbilityDecisionStop(ActorInfo);
+	}
+
+}
+
+void UAttributeComponent::MontageBranchPoint_AbilityDecisionStart()
+{
+	if (AnimatingAbility)
+	{
+		AActor * OwnerActor = GetOwner();
+		check(OwnerActor);
+
+		FGameplayAbilityActorInfo	ActorInfo;
+		ActorInfo.InitFromActor(OwnerActor);
+
+		AnimatingAbility->MontageBranchPoint_AbilityDecisionStart(ActorInfo);
+	}
+}
+
+bool UAttributeComponent::ActivateAbility(int32 AbilityIdx)
+{
+	AActor * OwnerActor = GetOwner();
+	check(OwnerActor);
+
+	FGameplayAbilityActorInfo	ActorInfo;
+	ActorInfo.InitFromActor(OwnerActor);
+
+	if (AllAbilities.IsValidIndex(AbilityIdx))
+	{
+		// Fixme
+		AllAbilities[AbilityIdx]->InputPressed(0, ActorInfo);
+	}
+
+	return false;
+}
+
+void UAttributeComponent::ServerTryActivateAbility_Implementation(class UGameplayAbility * AbilityToActivate)
+{
+	ensure(AbilityToActivate);
+
+	AActor * OwnerActor = GetOwner();
+	check(OwnerActor);
+
+	FGameplayAbilityActorInfo	ActorInfo;
+	ActorInfo.InitFromActor(OwnerActor);
+	
+	if (AbilityToActivate->CanActivateAbility(ActorInfo))
+	{
+		AbilityToActivate->ActivateAbility(ActorInfo);
+		ClientActivateAbilitySucceed(AbilityToActivate);
+	}
+	else
+	{
+		ClientActivateAbilityFailed(AbilityToActivate);
+	}
+}
+
+bool UAttributeComponent::ServerTryActivateAbility_Validate(class UGameplayAbility * AbilityToActivate)
+{
+	return true;
+}
+
+void UAttributeComponent::ClientActivateAbilityFailed_Implementation(class UGameplayAbility * AbilityToActivate)
+{
+
+}
+
+void UAttributeComponent::ClientActivateAbilitySucceed_Implementation(class UGameplayAbility * AbilityToActivate)
+{
+	ensure(AbilityToActivate);
+
+	AActor * OwnerActor = GetOwner();
+	check(OwnerActor);
+
+	FGameplayAbilityActorInfo	ActorInfo;
+	ActorInfo.InitFromActor(OwnerActor);
+
+	AbilityToActivate->ActivateAbility(ActorInfo);
+}
+
+
+#undef LOCTEXT_NAMESPACE
+
diff --git a/Engine/Source/Runtime/SkillSystem/Private/BlueprintPlayMontageAndWaitTaskProxy.cpp b/Engine/Source/Runtime/SkillSystem/Private/BlueprintPlayMontageAndWaitTaskProxy.cpp
new file mode 100644
--- /dev/null
+++ b/Engine/Source/Runtime/SkillSystem/Private/BlueprintPlayMontageAndWaitTaskProxy.cpp
@@ -0,0 +1,20 @@
+
+#include "SkillSystemModulePrivatePCH.h"
+
+UBlueprintPlayMontageAndWaitTaskProxy::UBlueprintPlayMontageAndWaitTaskProxy(const class FPostConstructInitializeProperties& PCIP)
+	: Super(PCIP)
+{
+
+}
+
+void UBlueprintPlayMontageAndWaitTaskProxy::OnMontageEnded(UAnimMontage* Montage, bool bInterrupted)
+{
+	if (bInterrupted)
+	{
+		OnInterrupted.Broadcast();
+	}
+	else
+	{
+		OnComplete.Broadcast();
+	}
+}
\ No newline at end of file
diff --git a/Engine/Source/Runtime/SkillSystem/Private/GameplayAbility.cpp b/Engine/Source/Runtime/SkillSystem/Private/GameplayAbility.cpp
new file mode 100644
--- /dev/null
+++ b/Engine/Source/Runtime/SkillSystem/Private/GameplayAbility.cpp
@@ -0,0 +1,222 @@
+// Copyright 1998-2013 Epic Games, Inc. All Rights Reserved.
+
+#include "SkillSystemModulePrivatePCH.h"
+
+// --------------------------------------------------------------------------------------------------------------------------------------------------------
+//
+//	UGameplayAbility
+//
+// --------------------------------------------------------------------------------------------------------------------------------------------------------
+
+
+UGameplayAbility::UGameplayAbility(const class FPostConstructInitializeProperties& PCIP)
+	: Super(PCIP)
+{
+	CostGameplayEffect = NULL;
+	CooldownGameplayEffect = NULL;
+
+}
+
+bool UGameplayAbility::CanActivateAbility(const FGameplayAbilityActorInfo ActorInfo) const
+{
+	// Check cooldowns and resource consumption immediately
+	if (!CheckCooldown(ActorInfo))
+	{
+		return false;
+	}
+
+	if (!CheckCost(ActorInfo))
+	{
+		return false;
+	}
+
+	return true;
+}
+
+bool UGameplayAbility::CommitAbility(const FGameplayAbilityActorInfo ActorInfo)
+{
+	if (!CommitCheck(ActorInfo))
+	{
+		return false;
+	}
+
+	CommitExecute(ActorInfo);
+	return true;
+}
+
+bool UGameplayAbility::CommitCheck(const FGameplayAbilityActorInfo ActorInfo)
+{
+	return CanActivateAbility(ActorInfo);
+}
+
+void UGameplayAbility::CommitExecute(const FGameplayAbilityActorInfo ActorInfo)
+{
+	if (!ActorInfo.AttributeComponent->IsNetSimulating())
+	{
+		ApplyCooldown(ActorInfo);
+
+		ApplyCost(ActorInfo);
+	}
+}
+
+bool UGameplayAbility::TryActivateAbility(const FGameplayAbilityActorInfo ActorInfo)
+{
+	// This should only come from button presses/local instigation
+	ENetRole NetMode = ActorInfo.Actor->Role;
+	ensure(NetMode != ROLE_SimulatedProxy);
+
+	// Always do a CanActivate check
+	if (!CanActivateAbility(ActorInfo))
+	{
+		return false;
+	}
+
+	// If we are the server or this is a client authorative 
+	if (ReplicationType == EGameplayAbilityReplication::Client || (NetMode == ROLE_Authority))
+	{
+		ActivateAbility(ActorInfo);
+	}
+	else if (ReplicationType == EGameplayAbilityReplication::Server)
+	{
+		ServerTryActivateAbility(ActorInfo);
+	}
+	else if (ReplicationType == EGameplayAbilityReplication::Predictive)
+	{
+		PredictiveActivateAbility(ActorInfo);
+		ServerTryActivateAbility(ActorInfo);
+	}
+
+	return true;	
+}
+
+void UGameplayAbility::ActivateAbility(FGameplayAbilityActorInfo ActorInfo)
+{
+	// Child classes may want to do stuff here...
+
+	if (CommitAbility(ActorInfo))		// Then commit to the ability...
+	{
+		//	Then do more stuff...
+	}
+}
+
+void UGameplayAbility::PredictiveActivateAbility(const FGameplayAbilityActorInfo ActorInfo)
+{
+
+}
+
+void UGameplayAbility::ServerTryActivateAbility(const FGameplayAbilityActorInfo ActorInfo)
+{
+	ActorInfo.AttributeComponent->ServerTryActivateAbility(this);
+}
+
+void UGameplayAbility::ClientActivateAbility(const FGameplayAbilityActorInfo ActorInfo)
+{
+	ActivateAbility(ActorInfo);
+}
+
+//----------------------------------------------------------------------
+
+void FGameplayAbilityActorInfo::InitFromActor(AActor *InActor)
+{
+	Actor = InActor;
+
+	USkeletalMeshComponent * SkelMeshComponent = InActor->FindComponentByClass<USkeletalMeshComponent>();
+	if (SkelMeshComponent)
+	{
+		this->AnimInstance = SkelMeshComponent->GetAnimInstance();
+	}
+
+	AttributeComponent = InActor->FindComponentByClass<UAttributeComponent>();
+/*
+	AActor *TestActor = Actor;
+	while(TestActor != NULL && this->InputComponent == NULL)
+	{
+		APawn *TestPawn = Cast<APawn>(TestActor);
+		if (TestPawn && TestPawn->GetController())
+		{
+			TestActor = TestPawn->GetController();
+		}
+
+		this->InputComponent = TestActor->FindComponentByClass<UInputComponent>();
+		TestActor = TestActor->GetOwner();
+	}
+*/
+}
+
+//----------------------------------------------------------------------
+
+void UGameplayAbility::InputPressed(int32 InputID, const FGameplayAbilityActorInfo ActorInfo)
+{
+	TryActivateAbility(ActorInfo);
+}
+
+void UGameplayAbility::InputReleased(int32 InputID, const FGameplayAbilityActorInfo ActorInfo)
+{
+	SKILL_LOG(Log, TEXT("InputReleased: %d"), InputID);
+}
+
+bool UGameplayAbility::CheckCooldown(const FGameplayAbilityActorInfo ActorInfo) const
+{
+	if (CooldownGameplayEffect)
+	{
+		check(ActorInfo.AttributeComponent.IsValid());
+		if (ActorInfo.AttributeComponent->HasAnyTags(CooldownGameplayEffect->OwnedTagsContainer))
+		{
+			return false;
+		}
+	}
+	return true;
+}
+
+void UGameplayAbility::ApplyCooldown(const FGameplayAbilityActorInfo ActorInfo)
+{
+	check(ActorInfo.AttributeComponent.IsValid());
+	if (CooldownGameplayEffect)
+	{
+		ActorInfo.AttributeComponent->ApplyGameplayEffectToSelf(CooldownGameplayEffect, 1.f, ActorInfo.Actor.Get());
+	}
+}
+
+bool UGameplayAbility::CheckCost(const FGameplayAbilityActorInfo ActorInfo) const
+{
+	check(ActorInfo.AttributeComponent.IsValid());
+	return ActorInfo.AttributeComponent->CanApplyAttributeModifiers(CostGameplayEffect, 1.f, ActorInfo.Actor.Get());
+}
+
+void UGameplayAbility::ApplyCost(const FGameplayAbilityActorInfo ActorInfo)
+{
+	check(ActorInfo.AttributeComponent.IsValid());
+	if (CostGameplayEffect)
+	{
+		ActorInfo.AttributeComponent->ApplyGameplayEffectToSelf(CostGameplayEffect, 1.f, ActorInfo.Actor.Get());
+	}
+}
+
+float UGameplayAbility::GetCooldownTimeRemaining(const FGameplayAbilityActorInfo ActorInfo) const
+{
+	check(ActorInfo.AttributeComponent.IsValid());
+	if (CooldownGameplayEffect)
+	{
+		TArray< float > Durations = ActorInfo.AttributeComponent->GetActiveEffectsTimeRemaining(FActiveGameplayEffectQuery(&CooldownGameplayEffect->OwnedTagsContainer));
+		if (Durations.Num() > 0)
+		{
+			Durations.Sort();
+			return Durations[0];
+		}
+	}
+
+	return 0.f;
+}
+
+
+// --------------------------------------------------------------------
+
+void UGameplayAbility::MontageBranchPoint_AbilityDecisionStop(const FGameplayAbilityActorInfo ActorInfo) const
+{
+
+}
+
+void UGameplayAbility::MontageBranchPoint_AbilityDecisionStart(const FGameplayAbilityActorInfo ActorInfo) const
+{
+
+}
\ No newline at end of file
diff --git a/Engine/Source/Runtime/SkillSystem/Private/GameplayAbilitySet.cpp b/Engine/Source/Runtime/SkillSystem/Private/GameplayAbilitySet.cpp
new file mode 100644
--- /dev/null
+++ b/Engine/Source/Runtime/SkillSystem/Private/GameplayAbilitySet.cpp
@@ -0,0 +1,17 @@
+// Copyright 1998-2013 Epic Games, Inc. All Rights Reserved.
+
+#include "SkillSystemModulePrivatePCH.h"
+
+// --------------------------------------------------------------------------------------------------------------------------------------------------------
+//
+//	UGameplayAbilitySet
+//
+// --------------------------------------------------------------------------------------------------------------------------------------------------------
+
+
+UGameplayAbilitySet::UGameplayAbilitySet(const class FPostConstructInitializeProperties& PCIP)
+	: Super(PCIP)
+{
+
+}
+
diff --git a/Engine/Source/Runtime/SkillSystem/Private/GameplayAbility_Instanced.cpp b/Engine/Source/Runtime/SkillSystem/Private/GameplayAbility_Instanced.cpp
new file mode 100644
--- /dev/null
+++ b/Engine/Source/Runtime/SkillSystem/Private/GameplayAbility_Instanced.cpp
@@ -0,0 +1,63 @@
+// Copyright 1998-2013 Epic Games, Inc. All Rights Reserved.
+
+#include "SkillSystemModulePrivatePCH.h"
+
+// --------------------------------------------------------------------------------------------------------------------------------------------------------
+//
+//	UGameplayAbility_Instanced
+//
+// --------------------------------------------------------------------------------------------------------------------------------------------------------
+
+
+UGameplayAbility_Instanced::UGameplayAbility_Instanced(const class FPostConstructInitializeProperties& PCIP)
+	: Super(PCIP)
+{
+}
+
+bool UGameplayAbility_Instanced::CanActivateAbility(const FGameplayAbilityActorInfo ActorInfo) const
+{
+	if (!Super::CanActivateAbility(ActorInfo))
+	{
+		return false;
+	}
+
+	// Call into blueprint (fixme: cache this)
+	static FName FuncName = FName(TEXT("K2_CanActivateAbility"));
+	UFunction* CanActivateFunction = GetClass()->FindFunctionByName(FuncName);
+
+	if (CanActivateFunction
+		&& CanActivateFunction->GetOuter()->IsA(UBlueprintGeneratedClass::StaticClass())
+		&& K2_CanActivateAbility(ActorInfo) == false)
+	{
+		SKILL_LOG(Log, TEXT("CanActivateAbility %s failed, blueprint refused"), *GetName());
+		return false;
+	}
+
+	return true;
+}
+
+void UGameplayAbility_Instanced::CommitExecute(const FGameplayAbilityActorInfo ActorInfo)
+{
+	Super::CommitExecute(ActorInfo);
+
+	// Call into blueprint (fixme: cache this)
+	static FName FuncName = FName(TEXT("K2_CommitAbility"));
+	UFunction* CanActivateFunction = GetClass()->FindFunctionByName(FuncName);
+
+	if (CanActivateFunction	&& CanActivateFunction->GetOuter()->IsA(UBlueprintGeneratedClass::StaticClass()))
+	{
+		K2_CommitExecute(ActorInfo);
+	}
+}
+
+void UGameplayAbility_Instanced::ActivateAbility(FGameplayAbilityActorInfo ActorInfo)
+{
+	// Call into blueprint (fixme: cache this)
+	static FName FuncName = FName(TEXT("K2_ActivateAbility"));
+	UFunction* CanActivateFunction = GetClass()->FindFunctionByName(FuncName);
+
+	if (CanActivateFunction	&& CanActivateFunction->GetOuter()->IsA(UBlueprintGeneratedClass::StaticClass()))
+	{
+		K2_ActivateAbility(ActorInfo);
+	}
+}
diff --git a/Engine/Source/Runtime/SkillSystem/Private/GameplayAbility_Montage.cpp b/Engine/Source/Runtime/SkillSystem/Private/GameplayAbility_Montage.cpp
new file mode 100644
--- /dev/null
+++ b/Engine/Source/Runtime/SkillSystem/Private/GameplayAbility_Montage.cpp
@@ -0,0 +1,63 @@
+// Copyright 1998-2013 Epic Games, Inc. All Rights Reserved.
+
+#include "SkillSystemModulePrivatePCH.h"
+
+// --------------------------------------------------------------------------------------------------------------------------------------------------------
+//
+//	UGameplayAbility_Montage
+//
+// --------------------------------------------------------------------------------------------------------------------------------------------------------
+
+UGameplayAbility_Montage::UGameplayAbility_Montage(const class FPostConstructInitializeProperties& PCIP)
+	: Super(PCIP)
+{
+	PlayRate = 1.f;
+	ReplicationType = EGameplayAbilityReplication::Server;
+}
+
+void UGameplayAbility_Montage::ActivateAbility(const FGameplayAbilityActorInfo OwnerInfo)
+{
+	if (!CommitAbility(OwnerInfo))
+	{
+		return;
+	}
+
+	if (MontageToPlay != NULL && OwnerInfo.AnimInstance != NULL && OwnerInfo.AnimInstance->GetActiveMontageInstance() == NULL)
+	{
+		TArray<FActiveGameplayEffectHandle>	AppliedEffects;
+
+		// Apply GameplayEffects
+		for (const UGameplayEffect * Effect : GameplayEffectsWhileAnimating)
+		{
+			FActiveGameplayEffectHandle Handle = OwnerInfo.AttributeComponent->ApplyGameplayEffectToSelf(Effect, 1.f, OwnerInfo.Actor.Get());
+			if (Handle.IsValid())
+			{
+				AppliedEffects.Add(Handle);
+			}
+		}
+
+		float const Duration = OwnerInfo.AnimInstance->Montage_Play(MontageToPlay, PlayRate);
+
+		FOnMontageEnded EndDelegate;
+		EndDelegate.BindUObject(this, &UGameplayAbility_Montage::OnMontageEnded, OwnerInfo, AppliedEffects);
+		OwnerInfo.AnimInstance->Montage_SetEndDelegate(EndDelegate);
+
+		if (SectionName != NAME_None)
+		{
+			OwnerInfo.AnimInstance->Montage_JumpToSection(SectionName);
+		}
+	}
+}
+
+void UGameplayAbility_Montage::OnMontageEnded(UAnimMontage* Montage, bool bInterrupted, const FGameplayAbilityActorInfo OwnerInfo, TArray<FActiveGameplayEffectHandle> AppliedEffects)
+{
+	// Remove any GameplayEffects that we applied
+	UAttributeComponent *OwnerComponent = OwnerInfo.AttributeComponent.Get();
+	if (OwnerComponent)
+	{
+		for (FActiveGameplayEffectHandle Handle : AppliedEffects)
+		{
+			OwnerComponent->RemoveActiveGameplayEffect(Handle);
+		}
+	}
+}
\ No newline at end of file
diff --git a/Engine/Source/Runtime/SkillSystem/Private/GameplayEffect.cpp b/Engine/Source/Runtime/SkillSystem/Private/GameplayEffect.cpp
--- a/Engine/Source/Runtime/SkillSystem/Private/GameplayEffect.cpp
+++ b/Engine/Source/Runtime/SkillSystem/Private/GameplayEffect.cpp
@@ -443,6 +443,39 @@ FString FAggregatorRef::ToString() const
 	return FString(TEXT("Invalid"));
 }
 
+bool FAggregatorRef::NetSerialize(FArchive& Ar, class UPackageMap* Map, bool& bOutSuccess)
+{
+	float EvaluatedMagnitude=0.f;
+
+	if (Ar.IsSaving())
+	{
+		if (IsValid())
+		{
+			EvaluatedMagnitude = Get()->Evaluate().Magnitude;
+		}
+	}
+
+	Ar << EvaluatedMagnitude;
+
+	if (Ar.IsLoading())
+	{
+		if (!IsValid())
+		{
+			SharedPtr = TSharedPtr<struct FAggregator>(new FAggregator);
+			WeakPtr = SharedPtr;
+		}
+
+		Get()->SetFromNetSerialize(EvaluatedMagnitude);
+	}
+	return true;
+}
+
+void FAggregator::SetFromNetSerialize(float SerializedValue)
+{
+	BaseData = FGameplayModifierData(SerializedValue, NULL);
+	CachedData = FGameplayModifierEvaluatedData(SerializedValue);
+}
+
 // --------------------------------------------------------------------------------------------------------------------------------------------------------
 //
 //	FAggregator
@@ -976,21 +1009,23 @@ void FActiveGameplayEffect::PreReplicatedRemove(const struct FActiveGameplayEffe
 
 void FActiveGameplayEffect::PostReplicatedAdd(const struct FActiveGameplayEffectsContainer &InArray)
 {
-	static const float MAX_DELTA_TIME = 3.f;
+	static const int32 MAX_DELTA_TIME = 3;
 
 	InArray.Owner->InvokeGameplayCueAdded(Spec);
 	// Was this actually just activated, or are we just finding out about it due to relevancy/join in progress?
-	float WorldTimeSeconds = InArray.Owner->GetWorld()->GetTimeSeconds();
-	if (InArray.Owner->GetWorld()->GetGameState<AGameState>())
-	{
-		WorldTimeSeconds = InArray.Owner->GetWorld()->GetGameState<AGameState>()->ElapsedTime;
-	}
-	float DeltaTime = WorldTimeSeconds - StartTime;
+	float WorldTimeSeconds = InArray.GetWorldTime();
+	int32 GameStateTime = InArray.GetGameStateTime();
+
+	int32 DeltaGameStateTime = GameStateTime - StartGameStateTime;	// How long we think the effect has been playing (only 1 second precision!)
 
-	if (WorldTimeSeconds > 0.f && FMath::Abs(DeltaTime) < MAX_DELTA_TIME)
+	if (GameStateTime > 0 && FMath::Abs(DeltaGameStateTime) < MAX_DELTA_TIME)
 	{
 		InArray.Owner->InvokeGameplayCueActivated(Spec);
 	}
+
+	// Set our local start time accordingly
+
+	StartWorldTime = WorldTimeSeconds - static_cast<float>(DeltaGameStateTime);
 }
 
 // --------------------------------------------------------------------------------------------------------------------------------------------------------
@@ -1289,10 +1324,10 @@ void FActiveGameplayEffectsContainer::OnPropertyAggregatorDirty(FAggregator* Agg
 	Owner->SetNumericAttribute(Attribute, NewPropertyValue);
 }
 
-FActiveGameplayEffect & FActiveGameplayEffectsContainer::CreateNewActiveGameplayEffect(const FGameplayEffectSpec &Spec, float GameTimeSeconds)
+FActiveGameplayEffect & FActiveGameplayEffectsContainer::CreateNewActiveGameplayEffect(const FGameplayEffectSpec &Spec)
 {
 	LastAssignedHandle = LastAssignedHandle.GetNextHandle();
-	FActiveGameplayEffect & NewEffect = *new (GameplayEffects)FActiveGameplayEffect(LastAssignedHandle, Spec, GameTimeSeconds);
+	FActiveGameplayEffect & NewEffect = *new (GameplayEffects)FActiveGameplayEffect(LastAssignedHandle, Spec, GetWorldTime(), GetGameStateTime());
 	
 	MarkItemDirty(NewEffect);
 	return NewEffect;
@@ -1335,6 +1370,24 @@ void FActiveGameplayEffectsContainer::PreDestroy()
 	}
 }
 
+int32 FActiveGameplayEffectsContainer::GetGameStateTime() const
+{
+	UWorld *World = Owner->GetWorld();
+	AGameState * GameState = World->GetGameState<AGameState>();
+	if (GameState)
+	{
+		return GameState->ElapsedTime;
+	}
+
+	return static_cast<int32>(World->GetTimeSeconds());
+}
+
+float FActiveGameplayEffectsContainer::GetWorldTime() const
+{
+	UWorld *World = Owner->GetWorld();
+	return World->GetTimeSeconds();
+}
+
 void FActiveGameplayEffectsContainer::TEMP_TickActiveEffects(float DeltaSeconds)
 {
 	if (!IsNetAuthority())
@@ -1344,13 +1397,7 @@ void FActiveGameplayEffectsContainer::TEMP_TickActiveEffects(float DeltaSeconds)
 		return;
 	}
 
-
-	float CurrentTime = Owner->GetWorld()->GetTimeSeconds();
-	// if we have a game state use the ElapsedTime so client server time replication is more accurate else use the time seconds of the game world.
-	if (Owner->GetWorld()->GetGameState<AGameState>())
-	{
-		CurrentTime = Owner->GetWorld()->GetGameState<AGameState>()->ElapsedTime;
-	}
+	float CurrentTime = GetWorldTime();
 
 	// effects may have been added outside of the tick during the last frame
 	if (bNeedToRecalculateStacks)
@@ -1379,7 +1426,7 @@ void FActiveGameplayEffectsContainer::TEMP_TickActiveEffects(float DeltaSeconds)
 		}
 
 		float Duration = Effect.GetDuration();
-		if (Duration > 0.f && Effect.StartTime + Effect.GetDuration() < CurrentTime)
+		if (Duration > 0.f && Effect.StartWorldTime + Effect.GetDuration() < CurrentTime)
 		{
 			Owner->InvokeGameplayCueRemoved(Effect.Spec);
 
@@ -1506,6 +1553,78 @@ void FActiveGameplayEffectsContainer::RecalculateStacking()
 	}
 }
 
+bool FActiveGameplayEffectsContainer::HasAnyTags(FGameplayTagContainer &Tags)
+{
+	SCOPE_CYCLE_COUNTER(STAT_GameplayEffectsHasAnyTag);
+
+	for (FActiveGameplayEffect &ActiveEffect : GameplayEffects)
+	{
+		// Fixme: check stacking rules!
+		if (ActiveEffect.Spec.Def->OwnedTagsContainer.HasAnyTag(Tags))
+		{
+			return true;
+		}
+	}
+
+	return false;
+}
+
+bool FActiveGameplayEffectsContainer::CanApplyAttributeModifiers(const UGameplayEffect *GameplayEffect, float Level, AActor *Instigator)
+{
+	SCOPE_CYCLE_COUNTER(STAT_GameplayEffectsCanApplyAttributeModifiers);
+
+	FGameplayEffectSpec	Spec(GameplayEffect, TSharedPtr<FGameplayEffectLevelSpec>(new FGameplayEffectLevelSpec(Level, GameplayEffect->LevelInfo, Instigator)), Owner->GetCurveDataOverride());
+	Spec.Def = GameplayEffect;
+	Spec.InstigatorStack.AddInstigator(Instigator);
+
+	ApplyActiveEffectsTo(Spec, FModifierQualifier().Type(EGameplayMod::IncomingGE));
+	
+	for (const FModifierSpec & Mod : Spec.Modifiers)
+	{
+		// It only makes sense to check additive operators
+		if (Mod.Info.ModifierOp == EGameplayModOp::Additive)
+		{
+			UAttributeSet * Set = Owner->GetAttributeSubobject(Mod.Info.Attribute.GetAttributeSetClass());
+			float CurrentValue = Mod.Info.Attribute.GetNumericValueChecked(Set);
+			float CostValue = Mod.Aggregator.Get()->Evaluate().Magnitude;
+
+			if (CurrentValue + CostValue < 0.f)
+			{
+				return false;
+			}
+		}
+	}
+	return true;
+}
+
+TArray<float> FActiveGameplayEffectsContainer::GetActiveEffectsTimeRemaining(const FActiveGameplayEffectQuery Query) const
+{
+	SCOPE_CYCLE_COUNTER(STAT_GameplayEffectsGetActiveEffectsData);
+
+	float CurrentTime = GetWorldTime();
+
+	TArray<float>	ReturnList;
+
+	for (const FActiveGameplayEffect &Effect : GameplayEffects)
+	{
+		if (Query.TagContainer)
+		{
+			if (!Effect.Spec.Def->OwnedTagsContainer.HasAnyTag(*Query.TagContainer))
+			{
+				continue;
+			}
+		}
+
+		float Elapsed = CurrentTime - Effect.StartWorldTime;
+		float Duration = Effect.GetDuration();
+
+		ReturnList.Add(Duration - Elapsed);
+	}
+
+	// Note: keep one return location to avoid copy operation.
+	return ReturnList;
+}
+
 // --------------------------------------------------------------------------------------------------------------------------------------------------------
 //
 //	FGameplayEffectLevelSpec
diff --git a/Engine/Source/Runtime/SkillSystem/Private/GameplayEffectStackingExtension_CappedNumberTest.cpp b/Engine/Source/Runtime/SkillSystem/Private/GameplayEffectStackingExtension_CappedNumberTest.cpp
--- a/Engine/Source/Runtime/SkillSystem/Private/GameplayEffectStackingExtension_CappedNumberTest.cpp
+++ b/Engine/Source/Runtime/SkillSystem/Private/GameplayEffectStackingExtension_CappedNumberTest.cpp
@@ -17,20 +17,20 @@ void UGameplayEffectStackingExtension_CappedNumberTest::CalculateStack(TArray<FA
 	int32 EffectiveCount = FMath::Min(CustomGameplayEffects.Num() + 1, MaxNumInStack);
 	
 	// find the most recent time one of these was applied
-	float StartTime = CurrentEffect.StartTime;
+	float StartTime = CurrentEffect.StartWorldTime;
 
 	for (FActiveGameplayEffect* Effect : CustomGameplayEffects)
 	{
-		StartTime = FMath::Max(StartTime, Effect->StartTime);
+		StartTime = FMath::Max(StartTime, Effect->StartWorldTime);
 	}
 
 	{
 		// set the start time to be equal to the most recent start time so that the stacked effects stick around if a new effect has been added
-		CurrentEffect.StartTime = StartTime;
+		CurrentEffect.StartWorldTime = StartTime;
 		int32 Idx = 0;
 		while (Idx < FMath::Min(MaxNumInStack - 1, CustomGameplayEffects.Num()))
 		{
-			CustomGameplayEffects[Idx]->StartTime = StartTime;
+			CustomGameplayEffects[Idx]->StartWorldTime = StartTime;
 			++Idx;
 		}
 
diff --git a/Engine/Source/Runtime/SkillSystem/Private/PlayMontageAndWaitAction.cpp b/Engine/Source/Runtime/SkillSystem/Private/PlayMontageAndWaitAction.cpp
new file mode 100644
--- /dev/null
+++ b/Engine/Source/Runtime/SkillSystem/Private/PlayMontageAndWaitAction.cpp
@@ -0,0 +1,29 @@
+// Copyright 1998-2014 Epic Games, Inc. All Rights Reserved.
+#pragma once
+
+#include "SkillSystemModulePrivatePCH.h"
+
+void FPlayMontageAndWaitAction::UpdateOperation(FLatentResponse& Response)
+{
+	//TimeRemaining -= Response.ElapsedTime();
+	//Response.FinishAndTriggerIf(TimeRemaining <= 0.0f, ExecutionFunction, OutputLink, CallbackTarget);
+}
+
+
+FPlayMontageAndWaitAction::FPlayMontageAndWaitAction(const FGameplayAbilityActorInfo InActorInfo, UAnimMontage * InMontageToPlay, const FLatentActionInfo& LatentInfo)
+	: ActorInfo(InActorInfo)
+	, MontageToPlay(InMontageToPlay)
+	, ExecutionFunction(LatentInfo.ExecutionFunction)
+	, OutputLink(LatentInfo.Linkage)
+	, CallbackTarget(LatentInfo.CallbackTarget)
+{
+	
+
+	ActorInfo.AnimInstance->Montage_Play(MontageToPlay, 1.f);
+
+	/*
+	FOnMontageEnded EndDelegate;
+	EndDelegate.BindUObject(this, &UGameplayAbility_Montage::OnMontageEnded, OwnerInfo, AppliedEffects);
+	OwnerInfo.AnimInstance->Montage_SetEndDelegate(EndDelegate);
+	*/
+}
\ No newline at end of file
diff --git a/Engine/Source/Runtime/SkillSystem/Private/SkillSystem.cpp b/Engine/Source/Runtime/SkillSystem/Private/SkillSystem.cpp
new file mode 100644
--- /dev/null
+++ b/Engine/Source/Runtime/SkillSystem/Private/SkillSystem.cpp
@@ -0,0 +1,8 @@
+// Copyright 1998-2014 Epic Games, Inc. All Rights Reserved.
+
+#include "SkillSystemModulePrivatePCH.h"
+
+DEFINE_STAT(STAT_GameplayEffectsHasAnyTag);
+DEFINE_STAT(STAT_GameplayEffectsTick);
+DEFINE_STAT(STAT_GameplayEffectsCanApplyAttributeModifiers);
+DEFINE_STAT(STAT_GameplayEffectsGetActiveEffectsData);
\ No newline at end of file
diff --git a/Engine/Source/Runtime/SkillSystem/Private/SkillSystemBlueprintLibrary.cpp b/Engine/Source/Runtime/SkillSystem/Private/SkillSystemBlueprintLibrary.cpp
new file mode 100644
--- /dev/null
+++ b/Engine/Source/Runtime/SkillSystem/Private/SkillSystemBlueprintLibrary.cpp
@@ -0,0 +1,54 @@
+// Copyright 1998-2014 Epic Games, Inc. All Rights Reserved.
+
+#include "SkillSystemModulePrivatePCH.h"
+#include "LatentActions.h"
+
+USkillSystemBlueprintLibrary::USkillSystemBlueprintLibrary(const class FPostConstructInitializeProperties& PCIP)
+: Super(PCIP)
+{
+}
+
+void USkillSystemBlueprintLibrary::StaticFuncPlayMontageAndWait(AActor* WorldContextObject, UAnimMontage * Montage, FLatentActionInfo LatentInfo)
+{
+	check(WorldContextObject);
+
+	if (UWorld* World = GEngine->GetWorldFromContextObject(WorldContextObject))
+	{
+		FLatentActionManager& LatentActionManager = World->GetLatentActionManager();
+		if (LatentActionManager.FindExistingAction<FPlayMontageAndWaitAction>(LatentInfo.CallbackTarget, LatentInfo.UUID) == NULL)
+		{
+			FGameplayAbilityActorInfo	ActorInfo;
+			ActorInfo.InitFromActor(WorldContextObject);
+
+			LatentActionManager.AddNewAction(LatentInfo.CallbackTarget, LatentInfo.UUID, new FPlayMontageAndWaitAction(ActorInfo, Montage, LatentInfo));
+		}
+	}
+}
+
+UBlueprintPlayMontageAndWaitTaskProxy* USkillSystemBlueprintLibrary::CreatePlayMontageAndWaitProxy(class UObject* WorldContextObject, UAnimMontage *MontageToPlay)
+{
+	check(WorldContextObject);
+
+	UGameplayAbility_Instanced * Ability = Cast<UGameplayAbility_Instanced>(WorldContextObject);
+	if (Ability)
+	{
+		AActor * ActorOwner = Cast<AActor>(Ability->GetOuter());
+
+		UBlueprintPlayMontageAndWaitTaskProxy * MyObj = NULL;
+		MyObj = NewObject<UBlueprintPlayMontageAndWaitTaskProxy>();
+
+		FOnMontageEnded EndDelegate;
+		EndDelegate.BindUObject(MyObj, &UBlueprintPlayMontageAndWaitTaskProxy::OnMontageEnded);
+
+		FGameplayAbilityActorInfo	ActorInfo;
+		ActorInfo.InitFromActor(ActorOwner);
+
+		ActorInfo.AnimInstance->Montage_Play(MontageToPlay, 1.f);
+		ActorInfo.AnimInstance->Montage_SetEndDelegate(EndDelegate);
+
+		return MyObj;
+	}
+
+	return NULL;
+	
+}
\ No newline at end of file
diff --git a/Engine/Source/Runtime/SkillSystem/Private/SkillSystemModulePrivatePCH.h b/Engine/Source/Runtime/SkillSystem/Private/SkillSystemModulePrivatePCH.h
--- a/Engine/Source/Runtime/SkillSystem/Private/SkillSystemModulePrivatePCH.h
+++ b/Engine/Source/Runtime/SkillSystem/Private/SkillSystemModulePrivatePCH.h
@@ -15,6 +15,9 @@
 #include "Engine/DataTable.h"
 #include "SkillSystemClasses.h"
 
+#include "SKillSystem.h"
 #include "SkillSystemModule.h"
 #include "SkillSystemLog.h"
 
+#include "PlayMontageAndWaitAction.h"
+#include "BlueprintPlayMontageAndWaitTaskProxy.h"
\ No newline at end of file
diff --git a/Engine/Source/Runtime/SkillSystem/Private/SkillSystemTestAttributeSet.cpp b/Engine/Source/Runtime/SkillSystem/Private/SkillSystemTestAttributeSet.cpp
--- a/Engine/Source/Runtime/SkillSystem/Private/SkillSystemTestAttributeSet.cpp
+++ b/Engine/Source/Runtime/SkillSystem/Private/SkillSystemTestAttributeSet.cpp
@@ -5,7 +5,9 @@
 USkillSystemTestAttributeSet::USkillSystemTestAttributeSet(const class FPostConstructInitializeProperties& PCIP)
 : Super(PCIP)
 {
-	Health = 100.f;
+	Health = MaxHealth = 100.f;
+	Mana = MaxMana = 100.f;
+	
 	Damage = 0.f;
 	CritChance = 0.f;
 	SpellDamage = 0.f;
@@ -100,3 +102,24 @@ void USkillSystemTestAttributeSet::PostAttributeModify(const struct FGameplayEff
 		//  -Doing it here makes a lot of sense to me, but we have legacy code in ::TakeDamage function, so some games may just want to punt to that pipeline from here.
 	}
 }
+
+
+void USkillSystemTestAttributeSet::GetLifetimeReplicatedProps(TArray< FLifetimeProperty > & OutLifetimeProps) const
+{
+	DOREPLIFETIME(USkillSystemTestAttributeSet, MaxHealth);
+	DOREPLIFETIME(USkillSystemTestAttributeSet, Health);
+	DOREPLIFETIME(USkillSystemTestAttributeSet, Mana);
+	DOREPLIFETIME(USkillSystemTestAttributeSet, MaxMana);
+
+	DOREPLIFETIME(USkillSystemTestAttributeSet, SpellDamage);
+	DOREPLIFETIME(USkillSystemTestAttributeSet, PhysicalDamage);
+
+	DOREPLIFETIME(USkillSystemTestAttributeSet, CritChance);
+	DOREPLIFETIME(USkillSystemTestAttributeSet, CritMultiplier);
+	DOREPLIFETIME(USkillSystemTestAttributeSet, ArmorDamageReduction);
+
+	DOREPLIFETIME(USkillSystemTestAttributeSet, DodgeChance);
+	DOREPLIFETIME(USkillSystemTestAttributeSet, LifeSteal);
+
+	DOREPLIFETIME(USkillSystemTestAttributeSet, Strength);
+}
\ No newline at end of file
diff --git a/Engine/Source/Runtime/SkillSystem/Private/SkillSystemTestPawn.cpp b/Engine/Source/Runtime/SkillSystem/Private/SkillSystemTestPawn.cpp
--- a/Engine/Source/Runtime/SkillSystem/Private/SkillSystemTestPawn.cpp
+++ b/Engine/Source/Runtime/SkillSystem/Private/SkillSystemTestPawn.cpp
@@ -9,6 +9,8 @@ ASkillSystemTestPawn::ASkillSystemTestPawn(const class FPostConstructInitializeP
 {
 	AttributeComponent = PCIP.CreateDefaultSubobject<UAttributeComponent>(this, ASkillSystemTestPawn::AttributeComponentName);
 	AttributeComponent->SetIsReplicated(true);
+
+	//DefaultAbilitySet = NULL;
 }
 
 void ASkillSystemTestPawn::PostInitializeComponents()
@@ -18,6 +20,13 @@ void ASkillSystemTestPawn::PostInitializeComponents()
 	Super::PostInitializeComponents();
 	GameplayCueHandler.Owner = this;
 	AttributeComponent->InitStats(USkillSystemTestAttributeSet::StaticClass(), NULL);
+
+	/*
+	if (DefaultAbilitySet != NULL)
+	{
+		AttributeComponent->InitializeAbilities(DefaultAbilitySet);
+	}
+	*/
 }
 
 void ASkillSystemTestPawn::GameplayCueActivated(const FGameplayTagContainer & GameplayCueTags, float NormalizedMagnitude)
diff --git a/Engine/Source/Runtime/SkillSystem/Public/PlayMontageAndWaitAction.h b/Engine/Source/Runtime/SkillSystem/Public/PlayMontageAndWaitAction.h
new file mode 100644
--- /dev/null
+++ b/Engine/Source/Runtime/SkillSystem/Public/PlayMontageAndWaitAction.h
@@ -0,0 +1,32 @@
+// Copyright 1998-2014 Epic Games, Inc. All Rights Reserved.
+#pragma once
+
+#include "LatentActions.h"
+
+// FDelayAction
+// A simple delay action; counts down and triggers it's output link when the time remaining falls to zero
+class FPlayMontageAndWaitAction : public FPendingLatentAction
+{
+public:
+	const FGameplayAbilityActorInfo ActorInfo;
+	class UAnimMontage * MontageToPlay;
+
+	FName ExecutionFunction;
+	int32 OutputLink;
+	FWeakObjectPtr CallbackTarget;
+
+	FPlayMontageAndWaitAction(const FGameplayAbilityActorInfo InActorInfo, UAnimMontage * InMontageToPlay, const FLatentActionInfo& LatentInfo);
+
+	virtual void UpdateOperation(FLatentResponse& Response);
+
+#if WITH_EDITOR
+	// Returns a human readable description of the latent operation's current state
+	virtual FString GetDescription() const OVERRIDE
+	{
+		//return FString::Printf(*NSLOCTEXT("DelayAction", "DelayActionTime", "Delay (%.3f seconds left)").ToString(), TimeRemaining);
+		return FString(TEXT("PlayMontageAndWait"));
+	}
+#endif
+};
+
+
diff --git a/Engine/Source/Runtime/SkillSystem/Public/SkillSystem.h b/Engine/Source/Runtime/SkillSystem/Public/SkillSystem.h
--- a/Engine/Source/Runtime/SkillSystem/Public/SkillSystem.h
+++ b/Engine/Source/Runtime/SkillSystem/Public/SkillSystem.h
@@ -5,4 +5,11 @@
 //#include "GameplayTagAssetInterface.h"
 //#include "GameplayTagsManager.h"
 //#include "GameplayTagContainer.h"
-#include "SkillSystemLog.h"
\ No newline at end of file
+#include "SkillSystemLog.h"
+
+DECLARE_STATS_GROUP(TEXT("SkillSystem"), STATGROUP_SkillSystem, STATCAT_Advanced);
+
+DECLARE_CYCLE_STAT_EXTERN(TEXT("GameplayEffectsHasAnyTagTime"), STAT_GameplayEffectsHasAnyTag, STATGROUP_SkillSystem, );
+DECLARE_CYCLE_STAT_EXTERN(TEXT("GameplayEffectsTick"), STAT_GameplayEffectsTick, STATGROUP_SkillSystem, );
+DECLARE_CYCLE_STAT_EXTERN(TEXT("CanApplyAttributeModifiers"), STAT_GameplayEffectsCanApplyAttributeModifiers, STATGROUP_SkillSystem, );
+DECLARE_CYCLE_STAT_EXTERN(TEXT("GetActiveEffectsData"), STAT_GameplayEffectsGetActiveEffectsData, STATGROUP_SkillSystem, );
\ No newline at end of file
