Temp clang fix

diff --git a/Engine/Source/Runtime/SkillSystem/Private/SkillSystemTestAttributeSet.cpp b/Engine/Source/Runtime/SkillSystem/Private/SkillSystemTestAttributeSet.cpp
--- a/Engine/Source/Runtime/SkillSystem/Private/SkillSystemTestAttributeSet.cpp
+++ b/Engine/Source/Runtime/SkillSystem/Private/SkillSystemTestAttributeSet.cpp
@@ -106,20 +106,22 @@ void USkillSystemTestAttributeSet::PostAttributeModify(const struct FGameplayEff
 
 void USkillSystemTestAttributeSet::GetLifetimeReplicatedProps(TArray< FLifetimeProperty > & OutLifetimeProps) const
 {
-	DOREPLIFETIME(USkillSystemTestAttributeSet, MaxHealth);
-	DOREPLIFETIME(USkillSystemTestAttributeSet, Health);
-	DOREPLIFETIME(USkillSystemTestAttributeSet, Mana);
-	DOREPLIFETIME(USkillSystemTestAttributeSet, MaxMana);
+	/*
+	DOREPLIFETIME( USkillSystemTestAttributeSet, MaxHealth);
+	DOREPLIFETIME( USkillSystemTestAttributeSet, Health);
+	DOREPLIFETIME( USkillSystemTestAttributeSet, Mana);
+	DOREPLIFETIME( USkillSystemTestAttributeSet, MaxMana);
 
-	DOREPLIFETIME(USkillSystemTestAttributeSet, SpellDamage);
-	DOREPLIFETIME(USkillSystemTestAttributeSet, PhysicalDamage);
+	DOREPLIFETIME( USkillSystemTestAttributeSet, SpellDamage);
+	DOREPLIFETIME( USkillSystemTestAttributeSet, PhysicalDamage);
 
-	DOREPLIFETIME(USkillSystemTestAttributeSet, CritChance);
-	DOREPLIFETIME(USkillSystemTestAttributeSet, CritMultiplier);
-	DOREPLIFETIME(USkillSystemTestAttributeSet, ArmorDamageReduction);
+	DOREPLIFETIME( USkillSystemTestAttributeSet, CritChance);
+	DOREPLIFETIME( USkillSystemTestAttributeSet, CritMultiplier);
+	DOREPLIFETIME( USkillSystemTestAttributeSet, ArmorDamageReduction);
 
-	DOREPLIFETIME(USkillSystemTestAttributeSet, DodgeChance);
-	DOREPLIFETIME(USkillSystemTestAttributeSet, LifeSteal);
+	DOREPLIFETIME( USkillSystemTestAttributeSet, DodgeChance);
+	DOREPLIFETIME( USkillSystemTestAttributeSet, LifeSteal);
 
-	DOREPLIFETIME(USkillSystemTestAttributeSet, Strength);
+	DOREPLIFETIME( USkillSystemTestAttributeSet, Strength);
+	*/
 }
\ No newline at end of file
