Allow GameplayEffects to pass other GameplayEffects to their target.

diff --git a/Engine/Source/Runtime/SkillSystem/Classes/AttributeComponent.h b/Engine/Source/Runtime/SkillSystem/Classes/AttributeComponent.h
--- a/Engine/Source/Runtime/SkillSystem/Classes/AttributeComponent.h
+++ b/Engine/Source/Runtime/SkillSystem/Classes/AttributeComponent.h
@@ -26,6 +26,8 @@ class SKILLSYSTEM_API UAttributeComponent : public UActorComponent
 {
 	GENERATED_UCLASS_BODY()
 
+	friend FGameplayEffectSpec;
+
 	virtual ~UAttributeComponent();
 
 	virtual void InitializeComponent() OVERRIDE;
@@ -63,7 +65,7 @@ class SKILLSYSTEM_API UAttributeComponent : public UActorComponent
 	// --------------------------------------------
 	// Primary outward facing API for other systems:
 	// --------------------------------------------
-	FActiveGameplayEffectHandle ApplyGameplayEffectSpecToTarget(OUT FGameplayEffectSpec &GameplayEffect, UAttributeComponent *Target, FModifierQualifier BaseQualifier = FModifierQualifier()) const;
+	FActiveGameplayEffectHandle ApplyGameplayEffectSpecToTarget(OUT FGameplayEffectSpec &GameplayEffect, UAttributeComponent *Target, FModifierQualifier BaseQualifier = FModifierQualifier());
 	FActiveGameplayEffectHandle ApplyGameplayEffectSpecToSelf(const FGameplayEffectSpec &GameplayEffect, FModifierQualifier BaseQualifier = FModifierQualifier());
 
 	UFUNCTION(BlueprintCallable, BlueprintAuthorityOnly, Category = GameplayEffects)
@@ -93,7 +95,7 @@ class SKILLSYSTEM_API UAttributeComponent : public UActorComponent
 	// Possibly useful but not primary API functions:
 	// --------------------------------------------
 
-	FActiveGameplayEffectHandle ApplyGameplayEffectSpecToTarget(UGameplayEffect *GameplayEffect, UAttributeComponent *Target, float Level = FGameplayEffectLevelSpec::INVALID_LEVEL, FModifierQualifier BaseQualifier = FModifierQualifier()) const;
+	FActiveGameplayEffectHandle ApplyGameplayEffectToTarget(UGameplayEffect *GameplayEffect, UAttributeComponent *Target, float Level = FGameplayEffectLevelSpec::INVALID_LEVEL, FModifierQualifier BaseQualifier = FModifierQualifier());
 
 	UFUNCTION(BlueprintCallable, BlueprintAuthorityOnly, Category = GameplayEffects, meta=(FriendlyName = "ApplyGameplayEffectToSelf"))
 	FActiveGameplayEffectHandle K2_ApplyGameplayEffectToSelf(const UGameplayEffect *GameplayEffect, float Level, AActor *Instigator);
diff --git a/Engine/Source/Runtime/SkillSystem/Classes/GameplayEffect.h b/Engine/Source/Runtime/SkillSystem/Classes/GameplayEffect.h
--- a/Engine/Source/Runtime/SkillSystem/Classes/GameplayEffect.h
+++ b/Engine/Source/Runtime/SkillSystem/Classes/GameplayEffect.h
@@ -22,6 +22,8 @@ DECLARE_MULTICAST_DELEGATE_ThreeParams(FOnAttributeGameplayEffectSpecExected, co
 	#define SKILL_AGG_DEBUG( Format, ... ) NULL
 #endif
 
+class UGameplayEffect;
+
 UENUM(BlueprintType)
 namespace EGameplayModOp
 {
@@ -52,7 +54,7 @@ namespace EGameplayMod
 		Attribute = 0,		// Modifies this Attributes
 		OutgoingGE,			// Modifies Outgoing Gameplay Effects (that modify this Attribute)
 		IncomingGE,			// Modifies Incoming Gameplay Effects (that modify this Attribute)
-		ActiveGE,
+		ActiveGE,			// Modifies currently active Gameplay Effects
 
 		// This must always be at the end
 		Max					UMETA(DisplayName="Invalid")
@@ -71,6 +73,7 @@ namespace EGameplayModEffect
 		Duration			= 0x02,		// Modifies duration of a GameplayEffect
 		ChanceApplyTarget	= 0x04,		// Modifies chance to apply GameplayEffect to target
 		ChanceApplyEffect	= 0x08,		// Modifies chance to apply GameplayEffect to GameplayEffect
+		LinkedGameplayEffect= 0x10,		// Adds a linked GameplayEffect to a GameplayEffect
 
 		// This must always be at the end
 		All					= 0xFF		UMETA(DisplayName="Invalid")
@@ -176,6 +179,7 @@ struct SKILLSYSTEM_API FGameplayModifierInfo
 		: ModifierType( EGameplayMod::Attribute )
 		, ModifierOp( EGameplayModOp::Additive )
 		, EffectType( EGameplayModEffect::Magnitude )
+		, TargetEffect( NULL )
 	{
 
 	}
@@ -199,6 +203,10 @@ struct SKILLSYSTEM_API FGameplayModifierInfo
 	UPROPERTY(EditDefaultsOnly, Category=GameplayModifier)
 	TEnumAsByte<EGameplayModEffect::Type> EffectType;
 
+	/** If we are linking a gameplay effect to another effect, this is the effect to link */
+	UPROPERTY(EditDefaultsOnly, Category=GameplayModifier)
+	UGameplayEffect* TargetEffect;
+
 	// The thing I modify requires these tags
 	UPROPERTY(EditDefaultsOnly, Category=GameplayModifier)
 	FGameplayTagContainer RequiredTags;
@@ -317,6 +325,10 @@ public:
 	UPROPERTY(EditDefaultsOnly, Category = Application, meta = (GameplayAttribute = "True"))
 	FScalableFloat	ChanceToApplyToGameplayEffect;
 
+	// other gameplay effects  that will be applied to the target of this effect
+	UPROPERTY(EditDefaultsOnly, BlueprintReadOnly, Category = GameplayEffect)
+	TArray<UGameplayEffect*> TargetEffects;
+
 	// Modify duration of CEs
 
 	// Modify MaxStacks of CEs (maybe... probably not)
@@ -1019,18 +1031,15 @@ private:
  */
 struct FModifierSpec
 {
-	FModifierSpec(const FGameplayModifierInfo &InInfo, TSharedPtr<FGameplayEffectLevelSpec> InLevel, const FGlobalCurveDataOverride *CurveData)
-		: Info(InInfo)
-		, Aggregator(new FAggregator(FGameplayModifierData(InInfo, CurveData), InLevel,  SKILL_AGG_DEBUG(TEXT("FModifierSpec: %s "), *InInfo.ToSimpleString() )))
-	{
-		Aggregator.Get()->RegisterLevelDependancies();
-	}
+	FModifierSpec(const FGameplayModifierInfo &InInfo, TSharedPtr<FGameplayEffectLevelSpec> InLevel, const FGlobalCurveDataOverride *CurveData, AActor *Owner = NULL, float Level = 0.f);
 
 	// Hard Ref to what we modify, this stuff is const and never changes
 	const FGameplayModifierInfo &Info;
 	
 	FAggregatorRef Aggregator;
 
+	TSharedPtr<FGameplayEffectSpec> TargetEffectSpec;
+
 	bool CanModifyInContext(const FModifierQualifier &QualifierContext) const;
 	// returns true if this GameplayEffect can modify Other, false otherwise
 	bool CanModifyModifier(FModifierSpec &Other, const FModifierQualifier &QualifierContext) const;
@@ -1073,7 +1082,7 @@ struct FGameplayEffectSpec
 		// If we initialize a GameplayEffectSpec with no level object passed in.
 	}
 
-	FGameplayEffectSpec( const UGameplayEffect * InDef, TSharedPtr<FGameplayEffectLevelSpec> InLevel, const FGlobalCurveDataOverride *CurveData );
+	FGameplayEffectSpec( const UGameplayEffect *InDef, AActor *Owner, float Level, const FGlobalCurveDataOverride *CurveData );
 	
 	UPROPERTY()
 	const UGameplayEffect * Def;
@@ -1086,13 +1095,18 @@ struct FGameplayEffectSpec
 
 	float GetDuration() const;
 	float GetPeriod() const;
+	EGameplayEffectStackingPolicy::Type GetStackingType() const;
 	float GetMagnitude(const FGameplayAttribute &Attribute) const;
 
+	// other effects that need to be applied to the target if this effect is successful
+	TArray< TSharedRef< FGameplayEffectSpec > > TargetEffectSpecs;
+
 	UPROPERTY()
 	FAggregatorRef	Duration;
 
 	FAggregatorRef	Period;
 
+	// How this combines with other gameplay effects
 	EGameplayEffectStackingPolicy::Type StackingPolicy;
 	FName StackedAttribName;
 
@@ -1103,15 +1117,13 @@ struct FGameplayEffectSpec
 
 	void MakeUnique();
 
-	void InitModifiers(const FGlobalCurveDataOverride *CurveData);
+	void InitModifiers(const FGlobalCurveDataOverride *CurveData, AActor *Owner, float Level);
 
-	int32 ApplyModifiersFrom(const FGameplayEffectSpec &InSpec, const FModifierQualifier &QualifierContext);
+	int32 ApplyModifiersFrom(FGameplayEffectSpec &InSpec, const FModifierQualifier &QualifierContext);
 	int32 ExecuteModifiersFrom(const FGameplayEffectSpec &InSpec, const FModifierQualifier &QualifierContext);
 
 	bool ShouldApplyAsSnapshot(const FModifierQualifier &QualifierContext) const;
 
-	EGameplayEffectStackingPolicy::Type GetStackingType() const;
-
 	FString ToSimpleString() const
 	{
 		return FString::Printf(TEXT("%s"), *Def->GetName());
@@ -1277,9 +1289,9 @@ struct FActiveGameplayEffectsContainer : public FFastArraySerializer
 	
 	FActiveGameplayEffect & CreateNewActiveGameplayEffect(const FGameplayEffectSpec &Spec);
 
-	void ApplyActiveEffectsTo(OUT FGameplayEffectSpec &Spec, const FModifierQualifier &QualifierContext) const;
+	void ApplyActiveEffectsTo(OUT FGameplayEffectSpec &Spec, const FModifierQualifier &QualifierContext);
 
-	void ApplySpecToActiveEffectsAndAttributes(const FGameplayEffectSpec &Spec, const FModifierQualifier &QualifierContext);
+	void ApplySpecToActiveEffectsAndAttributes(FGameplayEffectSpec &Spec, const FModifierQualifier &QualifierContext);
 		
 	void ExecuteActiveEffectsFrom(const FGameplayEffectSpec &Spec, const FModifierQualifier &QualifierContext);
 	
diff --git a/Engine/Source/Runtime/SkillSystem/Private/AttributeComponent.cpp b/Engine/Source/Runtime/SkillSystem/Private/AttributeComponent.cpp
--- a/Engine/Source/Runtime/SkillSystem/Private/AttributeComponent.cpp
+++ b/Engine/Source/Runtime/SkillSystem/Private/AttributeComponent.cpp
@@ -169,11 +169,11 @@ float UAttributeComponent::GetNumericAttribute(const FGameplayAttribute &Attribu
 }
 
 /** This is a helper function used in automated testing, I'm not sure how useful it will be to gamecode or blueprints */
-FActiveGameplayEffectHandle UAttributeComponent::ApplyGameplayEffectSpecToTarget(UGameplayEffect *GameplayEffect, UAttributeComponent *Target, float Level, FModifierQualifier BaseQualifier) const
+FActiveGameplayEffectHandle UAttributeComponent::ApplyGameplayEffectToTarget(UGameplayEffect *GameplayEffect, UAttributeComponent *Target, float Level, FModifierQualifier BaseQualifier)
 {
 	check(GameplayEffect);
 
-	FGameplayEffectSpec	Spec(GameplayEffect, TSharedPtr<FGameplayEffectLevelSpec>(new FGameplayEffectLevelSpec(Level, GameplayEffect->LevelInfo, GetOwner())), GetCurveDataOverride());
+	FGameplayEffectSpec	Spec(GameplayEffect, GetOwner(), Level, GetCurveDataOverride());
 	Spec.Def = GameplayEffect;
 	Spec.InstigatorStack.AddInstigator(GetOwner());
 	
@@ -189,7 +189,7 @@ FActiveGameplayEffectHandle UAttributeComponent::K2_ApplyGameplayEffectToSelf(co
 /** This is a helper function - it seems like this will be useful as a blueprint interface at the least, but Level parameter may need to be expanded */
 FActiveGameplayEffectHandle UAttributeComponent::ApplyGameplayEffectToSelf(const UGameplayEffect *GameplayEffect, float Level, AActor *Instigator, FModifierQualifier BaseQualifier)
 {
-	FGameplayEffectSpec	Spec(GameplayEffect, TSharedPtr<FGameplayEffectLevelSpec>(new FGameplayEffectLevelSpec(Level, GameplayEffect->LevelInfo, GetOwner())), GetCurveDataOverride());
+	FGameplayEffectSpec	Spec(GameplayEffect, GetOwner(), Level, GetCurveDataOverride());
 	Spec.Def = GameplayEffect;
 	Spec.InstigatorStack.AddInstigator(Instigator);
 
@@ -229,7 +229,7 @@ void UAttributeComponent::TEMP_ApplyActiveGameplayEffects()
 	}
 }
 
-FActiveGameplayEffectHandle UAttributeComponent::ApplyGameplayEffectSpecToTarget(OUT FGameplayEffectSpec &Spec, UAttributeComponent *Target, FModifierQualifier BaseQualifier) const
+FActiveGameplayEffectHandle UAttributeComponent::ApplyGameplayEffectSpecToTarget(OUT FGameplayEffectSpec &Spec, UAttributeComponent *Target, FModifierQualifier BaseQualifier)
 {
 	// Apply outgoing Effects to the Spec.
 	ActiveGameplayEffects.ApplyActiveEffectsTo(Spec, FModifierQualifier(BaseQualifier).Type(EGameplayMod::OutgoingGE));
@@ -312,12 +312,22 @@ FActiveGameplayEffectHandle UAttributeComponent::ApplyGameplayEffectSpecToSelf(c
 	
 	// Execute the GE at least once (if instant, this will execute once and be done. If persistent, it was added to ActiveGameplayEffects above)
 	
-	// Execute if this is an instant application effect or if it is a periodic effect (this would be the first execution of a repeating effect)
+	// Execute if this is an instant application effect
 	if (Spec.GetDuration() == UGameplayEffect::INSTANT_APPLICATION)
 	{
 		ExecuteGameplayEffect(*OurCopyOfSpec, FModifierQualifier(BaseQualifier).IgnoreHandle(MyHandle));
 	}
 
+	if (Spec.GetPeriod() != UGameplayEffect::NO_PERIOD && Spec.TargetEffectSpecs.Num() > 0)
+	{
+		SKILL_LOG(Warning, TEXT("%s is periodic but also applies GameplayEffects to its target. GameplayEffects will only be applied once, not every period."), *Spec.Def->GetPathName());
+	}
+	// todo: this is ignoring the returned handles, should we put them into a TArray and return all of the handles?
+	for (const TSharedRef<FGameplayEffectSpec> TargetSpec : Spec.TargetEffectSpecs)
+	{
+		ApplyGameplayEffectSpecToSelf(TargetSpec.Get(), BaseQualifier);
+	}
+
 	return MyHandle;
 }
 
diff --git a/Engine/Source/Runtime/SkillSystem/Private/GameplayEffect.cpp b/Engine/Source/Runtime/SkillSystem/Private/GameplayEffect.cpp
--- a/Engine/Source/Runtime/SkillSystem/Private/GameplayEffect.cpp
+++ b/Engine/Source/Runtime/SkillSystem/Private/GameplayEffect.cpp
@@ -56,13 +56,13 @@ bool UGameplayEffect::AreApplicationTagRequirementsSatisfied(const TSet<FName>&
 //
 // --------------------------------------------------------------------------------------------------------------------------------------------------------
 
-FGameplayEffectSpec::FGameplayEffectSpec(const UGameplayEffect * InDef, TSharedPtr<FGameplayEffectLevelSpec> InLevel, const FGlobalCurveDataOverride *CurveData)
-	: Def(InDef)
-	, ModifierLevel(InLevel)
-	, Duration(new FAggregator(InDef->Duration.MakeFinalizedCopy(CurveData), InLevel, SKILL_AGG_DEBUG(TEXT("%s Duration"), *InDef->GetName())))
-	, Period(new FAggregator(InDef->Period.MakeFinalizedCopy(CurveData), InLevel, SKILL_AGG_DEBUG(TEXT("%s Period"), *InDef->GetName())))
-	, StackingPolicy(EGameplayEffectStackingPolicy::Unlimited)
-	, StackedAttribName(NAME_None)
+FGameplayEffectSpec::FGameplayEffectSpec(const UGameplayEffect *InDef, AActor *Owner, float Level, const FGlobalCurveDataOverride *CurveData)
+: Def(InDef)
+, ModifierLevel(TSharedPtr<FGameplayEffectLevelSpec>(new FGameplayEffectLevelSpec(Level, InDef->LevelInfo, Owner)))
+, Duration(new FAggregator(InDef->Duration.MakeFinalizedCopy(CurveData), ModifierLevel, SKILL_AGG_DEBUG(TEXT("%s Duration"), *InDef->GetName())))
+, Period(new FAggregator(InDef->Period.MakeFinalizedCopy(CurveData), ModifierLevel, SKILL_AGG_DEBUG(TEXT("%s Period"), *InDef->GetName())))
+, StackingPolicy(EGameplayEffectStackingPolicy::Unlimited)
+, StackedAttribName(NAME_None)
 {
 	Duration.Get()->RegisterLevelDependancies();
 	Period.Get()->RegisterLevelDependancies();
@@ -145,10 +145,18 @@ FGameplayEffectSpec::FGameplayEffectSpec(const UGameplayEffect * InDef, TSharedP
 		}
 	}
 
-	InitModifiers(CurveData);
+	InitModifiers(CurveData, Owner, Level);
+
+	if (InDef)
+	{
+		for (UGameplayEffect* TargetDef : InDef->TargetEffects)
+		{
+			TargetEffectSpecs.Add(TSharedRef<FGameplayEffectSpec>(new FGameplayEffectSpec(TargetDef, Owner, Level, CurveData)));
+		}
+	}
 }
 
-void FGameplayEffectSpec::InitModifiers(const FGlobalCurveDataOverride *CurveData)
+void FGameplayEffectSpec::InitModifiers(const FGlobalCurveDataOverride *CurveData, AActor *Owner, float Level)
 {
 	check(Def);
 
@@ -165,7 +173,7 @@ void FGameplayEffectSpec::InitModifiers(const FGlobalCurveDataOverride *CurveDat
 		ModifierLevel->ApplyNewDef(ModInfo.LevelInfo, NewLevelSpec);
 
 		// This creates a new FModifierSpec that we own.
-		Modifiers.Emplace(FModifierSpec(ModInfo, NewLevelSpec, CurveData));
+		Modifiers.Emplace(FModifierSpec(ModInfo, NewLevelSpec, CurveData, Owner, Level));
 	}	
 }
 
@@ -177,7 +185,7 @@ void FGameplayEffectSpec::MakeUnique()
 	}
 }
 
-int32 FGameplayEffectSpec::ApplyModifiersFrom(const FGameplayEffectSpec &InSpec, const FModifierQualifier &QualifierContext)
+int32 FGameplayEffectSpec::ApplyModifiersFrom(FGameplayEffectSpec &InSpec, const FModifierQualifier &QualifierContext)
 {
 	SKILL_LOG_SCOPE(TEXT("FGameplayEffectSpec::ApplyModifiersFrom %s. InSpec: %s"), *this->ToSimpleString(), *InSpec.ToSimpleString());
 
@@ -224,6 +232,12 @@ int32 FGameplayEffectSpec::ApplyModifiersFrom(const FGameplayEffectSpec &InSpec,
 				NumApplied++;
 				break;
 			}
+
+			case EGameplayModEffect::LinkedGameplayEffect:
+			{
+				TargetEffectSpecs.Add(InMod.TargetEffectSpec.ToSharedRef());
+				break;
+			}
 		}
 	}
 
@@ -347,6 +361,17 @@ bool FModifierSpec::CanModifyModifier(FModifierSpec &Other, const FModifierQuali
 	return true;
 }
 
+FModifierSpec::FModifierSpec(const FGameplayModifierInfo &InInfo, TSharedPtr<FGameplayEffectLevelSpec> InLevel, const FGlobalCurveDataOverride *CurveData, AActor *Owner, float Level)
+: Info(InInfo)
+, Aggregator(new FAggregator(FGameplayModifierData(InInfo, CurveData), InLevel, SKILL_AGG_DEBUG(TEXT("FModifierSpec: %s "), *InInfo.ToSimpleString())))
+{
+	Aggregator.Get()->RegisterLevelDependancies();
+	if (InInfo.TargetEffect)
+	{
+		TargetEffectSpec = TSharedPtr<FGameplayEffectSpec>(new FGameplayEffectSpec(InInfo.TargetEffect, Owner, Level, CurveData));
+	}
+}
+
 void FModifierSpec::ApplyModTo(FModifierSpec &Other, bool TakeSnapshot) const
 {
 	Other.Aggregator.Get()->ApplyMod(this->Info.ModifierOp, this->Aggregator, TakeSnapshot);
@@ -1034,12 +1059,12 @@ void FActiveGameplayEffect::PostReplicatedAdd(const struct FActiveGameplayEffect
 //
 // --------------------------------------------------------------------------------------------------------------------------------------------------------
 
-void FActiveGameplayEffectsContainer::ApplyActiveEffectsTo(OUT FGameplayEffectSpec &Spec, const FModifierQualifier &QualifierContext) const
+void FActiveGameplayEffectsContainer::ApplyActiveEffectsTo(OUT FGameplayEffectSpec &Spec, const FModifierQualifier &QualifierContext)
 {
 	FActiveGameplayEffectHandle().IsValid();
 
 	SKILL_LOG_SCOPE(TEXT("ApplyActiveEffectsTo: %s %s"), *Spec.ToSimpleString(), *QualifierContext.ToString());
-	for (const FActiveGameplayEffect & ActiveEffect : GameplayEffects)
+	for (FActiveGameplayEffect & ActiveEffect : GameplayEffects)
 	{
 		// We dont want to use FModifierQualifier::TestTarget here, since we aren't the 'target'. We are applying stuff to Spec which will be applied to a target.
 		if (QualifierContext.IgnoreHandle().IsValid() && QualifierContext.IgnoreHandle() == ActiveEffect.Handle)
@@ -1051,7 +1076,7 @@ void FActiveGameplayEffectsContainer::ApplyActiveEffectsTo(OUT FGameplayEffectSp
 }
 
 /** This is the main function that applies/attaches a GameplayEffect on Attributes and ActiveGameplayEffects */
-void FActiveGameplayEffectsContainer::ApplySpecToActiveEffectsAndAttributes(const FGameplayEffectSpec &Spec, const FModifierQualifier &QualifierContext)
+void FActiveGameplayEffectsContainer::ApplySpecToActiveEffectsAndAttributes(FGameplayEffectSpec &Spec, const FModifierQualifier &QualifierContext)
 {
 	for (const FModifierSpec &Mod : Spec.Modifiers)
 	{
@@ -1573,7 +1598,7 @@ bool FActiveGameplayEffectsContainer::CanApplyAttributeModifiers(const UGameplay
 {
 	SCOPE_CYCLE_COUNTER(STAT_GameplayEffectsCanApplyAttributeModifiers);
 
-	FGameplayEffectSpec	Spec(GameplayEffect, TSharedPtr<FGameplayEffectLevelSpec>(new FGameplayEffectLevelSpec(Level, GameplayEffect->LevelInfo, Instigator)), Owner->GetCurveDataOverride());
+	FGameplayEffectSpec	Spec(GameplayEffect, Instigator, Level, Owner->GetCurveDataOverride());
 	Spec.Def = GameplayEffect;
 	Spec.InstigatorStack.AddInstigator(Instigator);
 
diff --git a/Engine/Source/Runtime/SkillSystem/Private/GameplayEffectTests.cpp b/Engine/Source/Runtime/SkillSystem/Private/GameplayEffectTests.cpp
--- a/Engine/Source/Runtime/SkillSystem/Private/GameplayEffectTests.cpp
+++ b/Engine/Source/Runtime/SkillSystem/Private/GameplayEffectTests.cpp
@@ -118,7 +118,7 @@ bool GameplayEffectsTest_InstantDamage(UWorld *World, FAutomationTestBase * Test
 		BaseDmgEffect->Modifiers[0].OwnedTags.AddTag(FName(TEXT("Damage.Basic")));
 		BaseDmgEffect->Duration.Value = UGameplayEffect::INSTANT_APPLICATION;
 
-		SourceComponent->ApplyGameplayEffectSpecToTarget(BaseDmgEffect, DestComponent, 1.f);
+		SourceComponent->ApplyGameplayEffectToTarget(BaseDmgEffect, DestComponent, 1.f);
 		
 		Test->TestTrue(SKILL_TEST_TEXT("Basic Instant Damage Applied"), (DestComponent->GetSet<USkillSystemTestAttributeSet>()->Health == (StartHealth + DamageValue)));
 		SKILL_LOG(Log, TEXT("Final Health: %.2f"), DestComponent->GetSet<USkillSystemTestAttributeSet>()->Health);
@@ -160,7 +160,7 @@ bool GameplayEffectsTest_InstantDamageRemap(UWorld *World, FAutomationTestBase *
 		BaseDmgEffect->Modifiers[0].OwnedTags.AddTag(FName(TEXT("Damage.Basic")));
 		BaseDmgEffect->Duration.Value = UGameplayEffect::INSTANT_APPLICATION;
 
-		SourceComponent->ApplyGameplayEffectSpecToTarget(BaseDmgEffect, DestComponent, 1.f);
+		SourceComponent->ApplyGameplayEffectToTarget(BaseDmgEffect, DestComponent, 1.f);
 
 		// Now we should have lost some health
 		{
@@ -213,7 +213,7 @@ bool GameplayEffectsTest_InstantDamage_Buffed(UWorld *World, FAutomationTestBase
 		BaseDmgEffect->Duration.SetValue(UGameplayEffect::INFINITE_DURATION);
 
 		// Apply to self
-		SourceComponent->ApplyGameplayEffectSpecToTarget(BaseDmgEffect, SourceComponent, 1.f);
+		SourceComponent->ApplyGameplayEffectToTarget(BaseDmgEffect, SourceComponent, 1.f);
 	}
 
 	// Apply Damage
@@ -230,7 +230,7 @@ bool GameplayEffectsTest_InstantDamage_Buffed(UWorld *World, FAutomationTestBase
 		BaseDmgEffect->Duration.Value = UGameplayEffect::INSTANT_APPLICATION;
 
 		// Apply to target
-		SourceComponent->ApplyGameplayEffectSpecToTarget(BaseDmgEffect, DestComponent, 1.f);
+		SourceComponent->ApplyGameplayEffectToTarget(BaseDmgEffect, DestComponent, 1.f);
 
 		float ExpectedValue = (StartHealth + (DamageValue * BonusDamageMultiplier));
 
@@ -279,7 +279,7 @@ bool GameplayEffectsTest_TemporaryDamage(UWorld *World, FAutomationTestBase * Te
 		BaseDmgEffect->Duration.Value = UGameplayEffect::INFINITE_DURATION;
 
 		// Apply to target
-		AppliedHandle = SourceComponent->ApplyGameplayEffectSpecToTarget(BaseDmgEffect, DestComponent, 1.f);
+		AppliedHandle = SourceComponent->ApplyGameplayEffectToTarget(BaseDmgEffect, DestComponent, 1.f);
 
 		float ExpectedValue = (StartHealth + DamageValue);
 
@@ -338,7 +338,7 @@ bool GameplayEffectsTest_TemporaryDamageBuffed(UWorld *World, FAutomationTestBas
 		BaseDmgEffect->Duration.Value = UGameplayEffect::INFINITE_DURATION;
 
 		// Apply to target
-		AppliedHandle = SourceComponent->ApplyGameplayEffectSpecToTarget(BaseDmgEffect, DestComponent, 1.f);
+		AppliedHandle = SourceComponent->ApplyGameplayEffectToTarget(BaseDmgEffect, DestComponent, 1.f);
 
 		float ExpectedValue = (StartHealth + DamageValue);
 
@@ -362,7 +362,7 @@ bool GameplayEffectsTest_TemporaryDamageBuffed(UWorld *World, FAutomationTestBas
 		BuffDmgEffect->Duration.Value = UGameplayEffect::INSTANT_APPLICATION;
 
 		// Apply to target
-		SourceComponent->ApplyGameplayEffectSpecToTarget(BuffDmgEffect, DestComponent, 1.f);
+		SourceComponent->ApplyGameplayEffectToTarget(BuffDmgEffect, DestComponent, 1.f);
 
 		float ExpectedValue = (StartHealth + (DamageValue * DamageBuffMultiplier));
 
@@ -430,7 +430,7 @@ bool GameplayEffectsTest_TemporaryDamageTemporaryBuff(UWorld *World, FAutomation
 		BaseDmgEffect->Duration.Value = UGameplayEffect::INFINITE_DURATION;
 
 		// Apply to target
-		AppliedHandle = SourceComponent->ApplyGameplayEffectSpecToTarget(BaseDmgEffect, DestComponent, 1.f);
+		AppliedHandle = SourceComponent->ApplyGameplayEffectToTarget(BaseDmgEffect, DestComponent, 1.f);
 
 		float ExpectedValue = (StartHealth + DamageValue);
 
@@ -455,7 +455,7 @@ bool GameplayEffectsTest_TemporaryDamageTemporaryBuff(UWorld *World, FAutomation
 		BuffDmgEffect->CopyPolicy = EGameplayEffectCopyPolicy::AlwaysLink;		// Force this to link, so that when we remove it it will go away to any modifier it was applied to
 
 		// Apply to target
-		AppliedHandleBuff = SourceComponent->ApplyGameplayEffectSpecToTarget(BuffDmgEffect, DestComponent, 1.f);
+		AppliedHandleBuff = SourceComponent->ApplyGameplayEffectToTarget(BuffDmgEffect, DestComponent, 1.f);
 
 		float ExpectedValue = (StartHealth + (DamageValue * DamageBuffMultiplier));
 
@@ -539,7 +539,7 @@ bool GameplayEffectsTest_LinkedBuffDestroy(UWorld *World, FAutomationTestBase *
 		BuffEffect->CopyPolicy = EGameplayEffectCopyPolicy::AlwaysLink;		// Always link so that when this is destroyed, health return to normal
 
 		// Apply to target
-		AppliedHandleBuff = SourceComponent->ApplyGameplayEffectSpecToTarget(BuffEffect, SourceComponent, 1.f);
+		AppliedHandleBuff = SourceComponent->ApplyGameplayEffectToTarget(BuffEffect, SourceComponent, 1.f);
 
 		Test->TestTrue(SKILL_TEST_TEXT("Number of Source GameplayEffect: %d == 1", SourceComponent->GetNumActiveGameplayEffect()), SourceComponent->GetNumActiveGameplayEffect() == 1);
 	}
@@ -557,7 +557,7 @@ bool GameplayEffectsTest_LinkedBuffDestroy(UWorld *World, FAutomationTestBase *
 		BaseDmgEffect->Duration.Value = UGameplayEffect::INFINITE_DURATION;
 
 		// Apply to target
-		AppliedHandle = SourceComponent->ApplyGameplayEffectSpecToTarget(BaseDmgEffect, DestComponent, 1.f);
+		AppliedHandle = SourceComponent->ApplyGameplayEffectToTarget(BaseDmgEffect, DestComponent, 1.f);
 
 		float ExpectedValue = (StartHealth + (DamageValue * DamageBuffMultiplier));
 		float ActualValue = DestComponent->GetSet<USkillSystemTestAttributeSet>()->Health;
@@ -626,7 +626,7 @@ bool GameplayEffectsTest_SnapshotBuffDestroy(UWorld *World, FAutomationTestBase
 		BuffEffect->CopyPolicy = EGameplayEffectCopyPolicy::AlwaysSnapshot;		// Always snapshot (though the default for oitgoing should already be snapshot - but this could change per project at some point)
 
 		// Apply to target
-		AppliedHandleBuff = SourceComponent->ApplyGameplayEffectSpecToTarget(BuffEffect, SourceComponent, 1.f);
+		AppliedHandleBuff = SourceComponent->ApplyGameplayEffectToTarget(BuffEffect, SourceComponent, 1.f);
 
 		Test->TestTrue(SKILL_TEST_TEXT("Number of Source GameplayEffect: %d == 1", SourceComponent->GetNumActiveGameplayEffect()), SourceComponent->GetNumActiveGameplayEffect() == 1);
 	}
@@ -644,7 +644,7 @@ bool GameplayEffectsTest_SnapshotBuffDestroy(UWorld *World, FAutomationTestBase
 		BaseDmgEffect->Duration.Value = UGameplayEffect::INFINITE_DURATION;
 
 		// Apply to target
-		AppliedHandle = SourceComponent->ApplyGameplayEffectSpecToTarget(BaseDmgEffect, DestComponent, 1.f);
+		AppliedHandle = SourceComponent->ApplyGameplayEffectToTarget(BaseDmgEffect, DestComponent, 1.f);
 
 		float ExpectedValue = (StartHealth + (DamageValue * DamageBuffMultiplier));
 		float ActualValue = DestComponent->GetSet<USkillSystemTestAttributeSet>()->Health;
@@ -715,7 +715,7 @@ bool GameplayEffectsTest_DurationBuff(UWorld *World, FAutomationTestBase * Test)
 		BaseDmgEffect->Duration.SetValue(BaseDuration);
 
 		// Apply to target
-		AppliedDamageHandle = SourceComponent->ApplyGameplayEffectSpecToTarget(BaseDmgEffect, DestComponent, 1.f);
+		AppliedDamageHandle = SourceComponent->ApplyGameplayEffectToTarget(BaseDmgEffect, DestComponent, 1.f);
 
 		float ExpectedHealthValue = (StartHealth + (DamageValue));
 		float ActualHealthValue = DestComponent->GetSet<USkillSystemTestAttributeSet>()->Health;
@@ -740,7 +740,7 @@ bool GameplayEffectsTest_DurationBuff(UWorld *World, FAutomationTestBase * Test)
 		DurationEffect->Duration.SetValue(UGameplayEffect::INFINITE_DURATION);
 
 		// Apply to target
-		AppliedDurationHandle = SourceComponent->ApplyGameplayEffectSpecToTarget(DurationEffect, DestComponent, 1.f);
+		AppliedDurationHandle = SourceComponent->ApplyGameplayEffectToTarget(DurationEffect, DestComponent, 1.f);
 
 		float ExpectedDuration = BaseDuration + DurationMod;
 		float ActualDuration = DestComponent->GetGameplayEffectDuration(AppliedDamageHandle);
@@ -809,7 +809,7 @@ bool GameplayEffectsTest_DamageBuffBuff_Basic(UWorld *World, FAutomationTestBase
 		BuffEffect->CopyPolicy = EGameplayEffectCopyPolicy::AlwaysLink;
 
 		// Apply to target
-		BuffBuffHandle = SourceComponent->ApplyGameplayEffectSpecToTarget(BuffEffect, SourceComponent, 1.f);
+		BuffBuffHandle = SourceComponent->ApplyGameplayEffectToTarget(BuffEffect, SourceComponent, 1.f);
 
 		Test->TestTrue(SKILL_TEST_TEXT("Number of Source GameplayEffect: %d == 1", SourceComponent->GetNumActiveGameplayEffect()), SourceComponent->GetNumActiveGameplayEffect() == 1);
 	}
@@ -827,7 +827,7 @@ bool GameplayEffectsTest_DamageBuffBuff_Basic(UWorld *World, FAutomationTestBase
 		BuffEffect->CopyPolicy = EGameplayEffectCopyPolicy::AlwaysSnapshot;
 
 		// Apply to target
-		BuffHandle = SourceComponent->ApplyGameplayEffectSpecToTarget(BuffEffect, SourceComponent, 1.f);
+		BuffHandle = SourceComponent->ApplyGameplayEffectToTarget(BuffEffect, SourceComponent, 1.f);
 		Test->TestTrue(SKILL_TEST_TEXT("Number of Source GameplayEffect: %d == 2", SourceComponent->GetNumActiveGameplayEffect()), SourceComponent->GetNumActiveGameplayEffect() == 2);
 
 		// Check that the buff was buffed
@@ -850,7 +850,7 @@ bool GameplayEffectsTest_DamageBuffBuff_Basic(UWorld *World, FAutomationTestBase
 		BaseDmgEffect->Duration.Value = UGameplayEffect::INFINITE_DURATION;
 
 		// Apply to target
-		DamageHandle = SourceComponent->ApplyGameplayEffectSpecToTarget(BaseDmgEffect, DestComponent, 1.f);
+		DamageHandle = SourceComponent->ApplyGameplayEffectToTarget(BaseDmgEffect, DestComponent, 1.f);
 
 		float ExpectedValue = (StartHealth + (DamageValue * (DamageBuffMultiplier + DamageBuffMultiplierBonus)));
 		float ActualValue = DestComponent->GetSet<USkillSystemTestAttributeSet>()->Health;
@@ -950,7 +950,7 @@ bool GameplayEffectsTest_DamageBuffBuff_FullLink(UWorld *World, FAutomationTestB
 		BuffEffect->CopyPolicy = EGameplayEffectCopyPolicy::AlwaysLink;
 
 		// Apply to target
-		BuffBuffHandle = SourceComponent->ApplyGameplayEffectSpecToTarget(BuffEffect, SourceComponent, 1.f);
+		BuffBuffHandle = SourceComponent->ApplyGameplayEffectToTarget(BuffEffect, SourceComponent, 1.f);
 
 		Test->TestTrue(SKILL_TEST_TEXT("Number of Source GameplayEffect: %d == 1", SourceComponent->GetNumActiveGameplayEffect()), SourceComponent->GetNumActiveGameplayEffect() == 1);
 	}
@@ -968,7 +968,7 @@ bool GameplayEffectsTest_DamageBuffBuff_FullLink(UWorld *World, FAutomationTestB
 		BuffEffect->CopyPolicy = EGameplayEffectCopyPolicy::AlwaysLink;
 
 		// Apply to target
-		BuffHandle = SourceComponent->ApplyGameplayEffectSpecToTarget(BuffEffect, SourceComponent, 1.f);
+		BuffHandle = SourceComponent->ApplyGameplayEffectToTarget(BuffEffect, SourceComponent, 1.f);
 		Test->TestTrue(SKILL_TEST_TEXT("Number of Source GameplayEffect: %d == 2", SourceComponent->GetNumActiveGameplayEffect()), SourceComponent->GetNumActiveGameplayEffect() == 2);
 
 		// Check that the buff was buffed
@@ -991,7 +991,7 @@ bool GameplayEffectsTest_DamageBuffBuff_FullLink(UWorld *World, FAutomationTestB
 		BaseDmgEffect->Duration.Value = UGameplayEffect::INFINITE_DURATION;
 
 		// Apply to target
-		DamageHandle = SourceComponent->ApplyGameplayEffectSpecToTarget(BaseDmgEffect, DestComponent, 1.f);
+		DamageHandle = SourceComponent->ApplyGameplayEffectToTarget(BaseDmgEffect, DestComponent, 1.f);
 
 		float ExpectedValue = (StartHealth + (DamageValue * (DamageBuffMultiplier + DamageBuffMultiplierBonus)));
 		float ActualValue = DestComponent->GetSet<USkillSystemTestAttributeSet>()->Health;
@@ -1088,7 +1088,7 @@ bool GameplayEffectsTest_DamageBuffBuff_FullSnapshot(UWorld *World, FAutomationT
 		BuffEffect->CopyPolicy = EGameplayEffectCopyPolicy::AlwaysSnapshot;
 
 		// Apply to target
-		BuffBuffHandle = SourceComponent->ApplyGameplayEffectSpecToTarget(BuffEffect, SourceComponent, 1.f);
+		BuffBuffHandle = SourceComponent->ApplyGameplayEffectToTarget(BuffEffect, SourceComponent, 1.f);
 
 		Test->TestTrue(SKILL_TEST_TEXT("Number of Source GameplayEffect: %d == 1", SourceComponent->GetNumActiveGameplayEffect()), SourceComponent->GetNumActiveGameplayEffect() == 1);
 	}
@@ -1106,7 +1106,7 @@ bool GameplayEffectsTest_DamageBuffBuff_FullSnapshot(UWorld *World, FAutomationT
 		BuffEffect->CopyPolicy = EGameplayEffectCopyPolicy::AlwaysSnapshot;
 
 		// Apply to target
-		BuffHandle = SourceComponent->ApplyGameplayEffectSpecToTarget(BuffEffect, SourceComponent, 1.f);
+		BuffHandle = SourceComponent->ApplyGameplayEffectToTarget(BuffEffect, SourceComponent, 1.f);
 		Test->TestTrue(SKILL_TEST_TEXT("Number of Source GameplayEffect: %d == 2", SourceComponent->GetNumActiveGameplayEffect()), SourceComponent->GetNumActiveGameplayEffect() == 2);
 
 		// Check that the buff was buffed
@@ -1129,7 +1129,7 @@ bool GameplayEffectsTest_DamageBuffBuff_FullSnapshot(UWorld *World, FAutomationT
 		BaseDmgEffect->Duration.Value = UGameplayEffect::INFINITE_DURATION;
 
 		// Apply to target
-		DamageHandle = SourceComponent->ApplyGameplayEffectSpecToTarget(BaseDmgEffect, DestComponent, 1.f);
+		DamageHandle = SourceComponent->ApplyGameplayEffectToTarget(BaseDmgEffect, DestComponent, 1.f);
 
 		float ExpectedValue = (StartHealth + (DamageValue * (DamageBuffMultiplier + DamageBuffMultiplierBonus)));
 		float ActualValue = DestComponent->GetSet<USkillSystemTestAttributeSet>()->Health;
@@ -1226,7 +1226,7 @@ bool GameplayEffectsTest_DamageBuffBuff_SnapshotLink(UWorld *World, FAutomationT
 		BuffEffect->CopyPolicy = EGameplayEffectCopyPolicy::AlwaysSnapshot;
 
 		// Apply to target
-		BuffBuffHandle = SourceComponent->ApplyGameplayEffectSpecToTarget(BuffEffect, SourceComponent, 1.f);
+		BuffBuffHandle = SourceComponent->ApplyGameplayEffectToTarget(BuffEffect, SourceComponent, 1.f);
 
 		Test->TestTrue(SKILL_TEST_TEXT("Number of Source GameplayEffect: %d == 1", SourceComponent->GetNumActiveGameplayEffect()), SourceComponent->GetNumActiveGameplayEffect() == 1);
 	}
@@ -1244,7 +1244,7 @@ bool GameplayEffectsTest_DamageBuffBuff_SnapshotLink(UWorld *World, FAutomationT
 		BuffEffect->CopyPolicy = EGameplayEffectCopyPolicy::AlwaysLink;
 
 		// Apply to target
-		BuffHandle = SourceComponent->ApplyGameplayEffectSpecToTarget(BuffEffect, SourceComponent, 1.f);
+		BuffHandle = SourceComponent->ApplyGameplayEffectToTarget(BuffEffect, SourceComponent, 1.f);
 		Test->TestTrue(SKILL_TEST_TEXT("Number of Source GameplayEffect: %d == 2", SourceComponent->GetNumActiveGameplayEffect()), SourceComponent->GetNumActiveGameplayEffect() == 2);
 
 		// Check that the buff was buffed
@@ -1267,7 +1267,7 @@ bool GameplayEffectsTest_DamageBuffBuff_SnapshotLink(UWorld *World, FAutomationT
 		BaseDmgEffect->Duration.Value = UGameplayEffect::INFINITE_DURATION;
 
 		// Apply to target
-		DamageHandle = SourceComponent->ApplyGameplayEffectSpecToTarget(BaseDmgEffect, DestComponent, 1.f);
+		DamageHandle = SourceComponent->ApplyGameplayEffectToTarget(BaseDmgEffect, DestComponent, 1.f);
 
 		float ExpectedValue = (StartHealth + (DamageValue * (DamageBuffMultiplier + DamageBuffMultiplierBonus)));
 		float ActualValue = DestComponent->GetSet<USkillSystemTestAttributeSet>()->Health;
@@ -1354,6 +1354,332 @@ bool TimerTest(UWorld *World, FAutomationTestBase * Test)
 	return true;
 }
 
+// Tests gameplay effects that apply other gameplay effects to the target
+bool GameplayEffectsTest_DamageAppliesBuff(UWorld *World, FAutomationTestBase * Test)
+{
+	const float StartHealth = 100.f;
+	const float DamageValue = 5.f;
+	const float DamageProtectionDivisor = 2.f;
+
+	ASkillSystemTestPawn *SourceActor = World->SpawnActor<ASkillSystemTestPawn>();
+	ASkillSystemTestPawn *DestActor = World->SpawnActor<ASkillSystemTestPawn>();
+
+	UProperty *DamageProperty = FindFieldChecked<UProperty>(USkillSystemTestAttributeSet::StaticClass(), GET_MEMBER_NAME_CHECKED(USkillSystemTestAttributeSet, Damage));
+
+	UAttributeComponent * SourceComponent = SourceActor->AttributeComponent;
+	UAttributeComponent * DestComponent = DestActor->AttributeComponent;
+	SourceComponent->GetSet<USkillSystemTestAttributeSet>()->Health = StartHealth;
+	DestComponent->GetSet<USkillSystemTestAttributeSet>()->Health = StartHealth;
+
+	// Apply damage and a buff that reduces incoming damage
+	{
+		SKILL_LOG_SCOPE(TEXT("Apply DamageBuff and InstantDamage"))
+
+		UGameplayEffect * BuffEffect = Cast<UGameplayEffect>(StaticConstructObject(UGameplayEffect::StaticClass(), GetTransientPackage(), FName(TEXT("DamageBuff"))));
+		BuffEffect->Modifiers.SetNum(1);
+		BuffEffect->Modifiers[0].Magnitude.SetValue(DamageProtectionDivisor);
+		BuffEffect->Modifiers[0].ModifierType = EGameplayMod::IncomingGE;
+		BuffEffect->Modifiers[0].ModifierOp = EGameplayModOp::Division;
+		BuffEffect->Modifiers[0].Attribute.SetUProperty(DamageProperty);
+		BuffEffect->Modifiers[0].OwnedTags.AddTag(FName(TEXT("Damage.Reduce")));
+		BuffEffect->Duration.SetValue(UGameplayEffect::INFINITE_DURATION);
+
+		UGameplayEffect * BaseDmgEffect = Cast<UGameplayEffect>(StaticConstructObject(UGameplayEffect::StaticClass(), GetTransientPackage(), FName(TEXT("Damage"))));
+		BaseDmgEffect->Modifiers.SetNum(1);
+		BaseDmgEffect->Modifiers[0].Magnitude.SetValue(DamageValue);
+		BaseDmgEffect->Modifiers[0].ModifierType = EGameplayMod::Attribute;
+		BaseDmgEffect->Modifiers[0].ModifierOp = EGameplayModOp::Additive;
+		BaseDmgEffect->Modifiers[0].Attribute.SetUProperty(DamageProperty);
+		BaseDmgEffect->Modifiers[0].OwnedTags.AddTag(FName(TEXT("Damage.Basic")));
+		BaseDmgEffect->Duration.Value = UGameplayEffect::INSTANT_APPLICATION;
+		BaseDmgEffect->TargetEffects.Add(BuffEffect);
+
+		// Apply to target
+		SourceComponent->ApplyGameplayEffectToTarget(BaseDmgEffect, DestComponent, 1.f);
+
+		float ExpectedValue = (StartHealth - DamageValue);
+		float ActualValue = DestComponent->GetSet<USkillSystemTestAttributeSet>()->Health;
+
+		Test->TestTrue(SKILL_TEST_TEXT("Buff Instant Damage Applied"), (ActualValue == ExpectedValue));
+		SKILL_LOG(Log, TEXT("Final Health: %.2f"), DestComponent->GetSet<USkillSystemTestAttributeSet>()->Health);
+	}
+
+	// Apply Damage
+	{
+		UGameplayEffect * BaseDmgEffect = Cast<UGameplayEffect>(StaticConstructObject(UGameplayEffect::StaticClass(), GetTransientPackage(), FName(TEXT("Damage"))));
+		BaseDmgEffect->Modifiers.SetNum(1);
+		BaseDmgEffect->Modifiers[0].Magnitude.SetValue(DamageValue);
+		BaseDmgEffect->Modifiers[0].ModifierType = EGameplayMod::Attribute;
+		BaseDmgEffect->Modifiers[0].ModifierOp = EGameplayModOp::Additive;
+		BaseDmgEffect->Modifiers[0].Attribute.SetUProperty(DamageProperty);
+		BaseDmgEffect->Modifiers[0].OwnedTags.AddTag(FName(TEXT("Damage.Basic")));
+		BaseDmgEffect->Duration.Value = UGameplayEffect::INSTANT_APPLICATION;
+
+		// Apply to target
+		SourceComponent->ApplyGameplayEffectToTarget(BaseDmgEffect, DestComponent, 1.f);
+
+		float ExpectedValue = (StartHealth - DamageValue - (DamageValue / DamageProtectionDivisor));
+		float ActualValue = DestComponent->GetSet<USkillSystemTestAttributeSet>()->Health;
+
+		Test->TestTrue(SKILL_TEST_TEXT("Buff Instant Damage Applied"), (ActualValue == ExpectedValue));
+		SKILL_LOG(Log, TEXT("Final Health: %.2f"), DestComponent->GetSet<USkillSystemTestAttributeSet>()->Health);
+	}
+
+	World->EditorDestroyActor(SourceActor, false);
+	World->EditorDestroyActor(DestActor, false);
+
+	return true;
+}
+
+// Source has a buff that applies a buff to the target of all damage effects
+bool GameplayEffectsTest_BuffAppliesBuff(UWorld *World, FAutomationTestBase * Test)
+{
+	const float StartHealth = 100.f;
+	const float DamageValue = 5.f;
+	const float DamageProtectionDivisor = 2.f;
+
+	ASkillSystemTestPawn *SourceActor = World->SpawnActor<ASkillSystemTestPawn>();
+	ASkillSystemTestPawn *DestActor = World->SpawnActor<ASkillSystemTestPawn>();
+
+	UProperty *DamageProperty = FindFieldChecked<UProperty>(USkillSystemTestAttributeSet::StaticClass(), GET_MEMBER_NAME_CHECKED(USkillSystemTestAttributeSet, Damage));
+
+	UAttributeComponent * SourceComponent = SourceActor->AttributeComponent;
+	UAttributeComponent * DestComponent = DestActor->AttributeComponent;
+	SourceComponent->GetSet<USkillSystemTestAttributeSet>()->Health = StartHealth;
+	DestComponent->GetSet<USkillSystemTestAttributeSet>()->Health = StartHealth;
+
+	// Apply damage and a buff that reduces incoming damage
+	{
+		SKILL_LOG_SCOPE(TEXT("Apply DamageBuff"))
+
+		UGameplayEffect * BuffEffect = Cast<UGameplayEffect>(StaticConstructObject(UGameplayEffect::StaticClass(), GetTransientPackage(), FName(TEXT("DamageBuff"))));
+		BuffEffect->Modifiers.SetNum(1);
+		BuffEffect->Modifiers[0].Magnitude.SetValue(DamageProtectionDivisor);
+		BuffEffect->Modifiers[0].ModifierType = EGameplayMod::IncomingGE;
+		BuffEffect->Modifiers[0].ModifierOp = EGameplayModOp::Division;
+		BuffEffect->Modifiers[0].Attribute.SetUProperty(DamageProperty);
+		BuffEffect->Modifiers[0].OwnedTags.AddTag(FName(TEXT("Damage.Reduce")));
+		BuffEffect->Duration.SetValue(UGameplayEffect::INFINITE_DURATION);
+
+		UGameplayEffect * DummyBuffEffect = Cast<UGameplayEffect>(StaticConstructObject(UGameplayEffect::StaticClass(), GetTransientPackage(), FName(TEXT("Dummy"))));
+		DummyBuffEffect->Modifiers.SetNum(1);
+		DummyBuffEffect->Modifiers[0].ModifierType = EGameplayMod::OutgoingGE;
+		DummyBuffEffect->Modifiers[0].EffectType = EGameplayModEffect::LinkedGameplayEffect;
+		DummyBuffEffect->Modifiers[0].RequiredTags.AddTag(FName(TEXT("Damage.Buffable")));
+		DummyBuffEffect->Modifiers[0].TargetEffect = BuffEffect;
+		DummyBuffEffect->Duration.SetValue(UGameplayEffect::INFINITE_DURATION);
+
+		// Apply to target
+		SourceComponent->ApplyGameplayEffectToTarget(DummyBuffEffect, SourceComponent, 1.f);
+	}
+
+	// apply damage to source to make sure it didn't get the protection buff
+	{
+		UGameplayEffect * BaseDmgEffect = Cast<UGameplayEffect>(StaticConstructObject(UGameplayEffect::StaticClass(), GetTransientPackage(), FName(TEXT("Damage"))));
+		BaseDmgEffect->Modifiers.SetNum(1);
+		BaseDmgEffect->Modifiers[0].Magnitude.SetValue(DamageValue);
+		BaseDmgEffect->Modifiers[0].ModifierType = EGameplayMod::Attribute;
+		BaseDmgEffect->Modifiers[0].ModifierOp = EGameplayModOp::Additive;
+		BaseDmgEffect->Modifiers[0].Attribute.SetUProperty(DamageProperty);
+		BaseDmgEffect->Modifiers[0].OwnedTags.AddTag(FName(TEXT("Damage.Basic")));
+		BaseDmgEffect->Duration.Value = UGameplayEffect::INSTANT_APPLICATION;
+
+		// Apply to target
+		DestComponent->ApplyGameplayEffectToTarget(BaseDmgEffect, SourceComponent, 1.f);
+
+		float ExpectedValue = StartHealth - DamageValue;
+		float ActualValue = SourceComponent->GetSet<USkillSystemTestAttributeSet>()->Health;
+
+		Test->TestTrue(SKILL_TEST_TEXT("Sending buffs test"), (ActualValue == ExpectedValue));
+		SKILL_LOG(Log, TEXT("Final Health: %.2f"), SourceComponent->GetSet<USkillSystemTestAttributeSet>()->Health);
+	}
+
+	// Apply Damage
+	{
+		UGameplayEffect * BaseDmgEffect = Cast<UGameplayEffect>(StaticConstructObject(UGameplayEffect::StaticClass(), GetTransientPackage(), FName(TEXT("Damage"))));
+		BaseDmgEffect->Modifiers.SetNum(1);
+		BaseDmgEffect->Modifiers[0].Magnitude.SetValue(DamageValue);
+		BaseDmgEffect->Modifiers[0].ModifierType = EGameplayMod::Attribute;
+		BaseDmgEffect->Modifiers[0].ModifierOp = EGameplayModOp::Additive;
+		BaseDmgEffect->Modifiers[0].Attribute.SetUProperty(DamageProperty);
+		BaseDmgEffect->Modifiers[0].OwnedTags.AddTag(FName(TEXT("Damage.Buffable")));
+		BaseDmgEffect->Duration.Value = UGameplayEffect::INSTANT_APPLICATION;
+
+		// Apply to target
+		SourceComponent->ApplyGameplayEffectToTarget(BaseDmgEffect, DestComponent, 1.f);
+
+		float ExpectedValue = StartHealth - DamageValue;
+		float ActualValue = DestComponent->GetSet<USkillSystemTestAttributeSet>()->Health;
+
+		Test->TestTrue(SKILL_TEST_TEXT("Buff Instant Damage Applied"), (ActualValue == ExpectedValue));
+		SKILL_LOG(Log, TEXT("Final Health: %.2f"), DestComponent->GetSet<USkillSystemTestAttributeSet>()->Health);
+	}
+
+	// Apply Damage
+	{
+		UGameplayEffect * BaseDmgEffect = Cast<UGameplayEffect>(StaticConstructObject(UGameplayEffect::StaticClass(), GetTransientPackage(), FName(TEXT("Damage"))));
+		BaseDmgEffect->Modifiers.SetNum(1);
+		BaseDmgEffect->Modifiers[0].Magnitude.SetValue(DamageValue);
+		BaseDmgEffect->Modifiers[0].ModifierType = EGameplayMod::Attribute;
+		BaseDmgEffect->Modifiers[0].ModifierOp = EGameplayModOp::Additive;
+		BaseDmgEffect->Modifiers[0].Attribute.SetUProperty(DamageProperty);
+		BaseDmgEffect->Modifiers[0].OwnedTags.AddTag(FName(TEXT("Damage.Basic")));
+		BaseDmgEffect->Duration.Value = UGameplayEffect::INSTANT_APPLICATION;
+
+		// Apply to target
+		SourceComponent->ApplyGameplayEffectToTarget(BaseDmgEffect, DestComponent, 1.f);
+
+		float ExpectedValue = (StartHealth - DamageValue - (DamageValue / DamageProtectionDivisor));
+		float ActualValue = DestComponent->GetSet<USkillSystemTestAttributeSet>()->Health;
+
+		Test->TestTrue(SKILL_TEST_TEXT("Instant Damage Applied"), (ActualValue == ExpectedValue));
+		SKILL_LOG(Log, TEXT("Final Health: %.2f"), DestComponent->GetSet<USkillSystemTestAttributeSet>()->Health);
+	}
+
+	// Apply Damage again to make sure that the buff only applied once
+	{
+		UGameplayEffect * BaseDmgEffect = Cast<UGameplayEffect>(StaticConstructObject(UGameplayEffect::StaticClass(), GetTransientPackage(), FName(TEXT("Damage"))));
+		BaseDmgEffect->Modifiers.SetNum(1);
+		BaseDmgEffect->Modifiers[0].Magnitude.SetValue(DamageValue);
+		BaseDmgEffect->Modifiers[0].ModifierType = EGameplayMod::Attribute;
+		BaseDmgEffect->Modifiers[0].ModifierOp = EGameplayModOp::Additive;
+		BaseDmgEffect->Modifiers[0].Attribute.SetUProperty(DamageProperty);
+		BaseDmgEffect->Modifiers[0].OwnedTags.AddTag(FName(TEXT("Damage.Basic")));
+		BaseDmgEffect->Duration.Value = UGameplayEffect::INSTANT_APPLICATION;
+
+		// Apply to target
+		SourceComponent->ApplyGameplayEffectToTarget(BaseDmgEffect, DestComponent, 1.f);
+
+		float ExpectedValue = (StartHealth - DamageValue - (2 * DamageValue / DamageProtectionDivisor));
+		float ActualValue = DestComponent->GetSet<USkillSystemTestAttributeSet>()->Health;
+
+		Test->TestTrue(SKILL_TEST_TEXT("Instant Damage Applied"), (ActualValue == ExpectedValue));
+		SKILL_LOG(Log, TEXT("Final Health: %.2f"), DestComponent->GetSet<USkillSystemTestAttributeSet>()->Health);
+	}
+
+	World->EditorDestroyActor(SourceActor, false);
+	World->EditorDestroyActor(DestActor, false);
+
+	return true;
+}
+
+bool GameplayEffectsTest_BuffIndirection(UWorld *World, FAutomationTestBase * Test)
+{
+	const float StartHealth = 100.f;
+	const float DamageValue = 5.f;
+	const float DamageProtectionDivisor = 2.f;
+
+	ASkillSystemTestPawn *SourceActor = World->SpawnActor<ASkillSystemTestPawn>();
+	ASkillSystemTestPawn *DestActor = World->SpawnActor<ASkillSystemTestPawn>();
+
+	UProperty *DamageProperty = FindFieldChecked<UProperty>(USkillSystemTestAttributeSet::StaticClass(), GET_MEMBER_NAME_CHECKED(USkillSystemTestAttributeSet, Damage));
+
+	UAttributeComponent * SourceComponent = SourceActor->AttributeComponent;
+	UAttributeComponent * DestComponent = DestActor->AttributeComponent;
+	SourceComponent->GetSet<USkillSystemTestAttributeSet>()->Health = StartHealth;
+	DestComponent->GetSet<USkillSystemTestAttributeSet>()->Health = StartHealth;
+
+	// Apply damage and a buff that reduces incoming damage
+	{
+		SKILL_LOG_SCOPE(TEXT("Apply DamageBuff and InstantDamage"))
+
+		UGameplayEffect * BuffEffect = Cast<UGameplayEffect>(StaticConstructObject(UGameplayEffect::StaticClass(), GetTransientPackage(), FName(TEXT("DamageBuff"))));
+		BuffEffect->Modifiers.SetNum(1);
+		BuffEffect->Modifiers[0].Magnitude.SetValue(DamageProtectionDivisor);
+		BuffEffect->Modifiers[0].ModifierType = EGameplayMod::IncomingGE;
+		BuffEffect->Modifiers[0].ModifierOp = EGameplayModOp::Division;
+		BuffEffect->Modifiers[0].Attribute.SetUProperty(DamageProperty);
+		BuffEffect->Modifiers[0].OwnedTags.AddTag(FName(TEXT("Damage.Reduce")));
+		BuffEffect->Duration.SetValue(UGameplayEffect::INFINITE_DURATION);
+
+		UGameplayEffect * DummyBuffEffect = Cast<UGameplayEffect>(StaticConstructObject(UGameplayEffect::StaticClass(), GetTransientPackage(), FName(TEXT("Dummy"))));
+		DummyBuffEffect->Modifiers.SetNum(1);
+		DummyBuffEffect->Modifiers[0].ModifierType = EGameplayMod::OutgoingGE;
+		DummyBuffEffect->Modifiers[0].EffectType = EGameplayModEffect::LinkedGameplayEffect;
+		DummyBuffEffect->Modifiers[0].TargetEffect = BuffEffect;
+		DummyBuffEffect->Modifiers[0].RequiredTags.AddTag("Damage.Buffable");
+		DummyBuffEffect->Duration.SetValue(UGameplayEffect::INFINITE_DURATION);
+
+		UGameplayEffect * DummyBuffEffect2 = Cast<UGameplayEffect>(StaticConstructObject(UGameplayEffect::StaticClass(), GetTransientPackage(), FName(TEXT("Dummy2"))));
+		DummyBuffEffect2->Modifiers.SetNum(1);
+		DummyBuffEffect2->Modifiers[0].ModifierType = EGameplayMod::OutgoingGE;
+		DummyBuffEffect2->Modifiers[0].EffectType = EGameplayModEffect::LinkedGameplayEffect;
+		DummyBuffEffect2->Modifiers[0].TargetEffect = DummyBuffEffect;
+		DummyBuffEffect2->Duration.SetValue(UGameplayEffect::INFINITE_DURATION);
+
+		// Apply to target
+		SourceComponent->ApplyGameplayEffectToTarget(DummyBuffEffect2, SourceComponent, 1.f);
+	}
+
+	// Apply Damage
+	{
+		UGameplayEffect * BaseDmgEffect = Cast<UGameplayEffect>(StaticConstructObject(UGameplayEffect::StaticClass(), GetTransientPackage(), FName(TEXT("Damage"))));
+		BaseDmgEffect->Modifiers.SetNum(1);
+		BaseDmgEffect->Modifiers[0].Magnitude.SetValue(DamageValue);
+		BaseDmgEffect->Modifiers[0].ModifierType = EGameplayMod::Attribute;
+		BaseDmgEffect->Modifiers[0].ModifierOp = EGameplayModOp::Additive;
+		BaseDmgEffect->Modifiers[0].Attribute.SetUProperty(DamageProperty);
+		BaseDmgEffect->Modifiers[0].OwnedTags.AddTag(FName(TEXT("Damage.Basic")));
+		BaseDmgEffect->Duration.Value = UGameplayEffect::INSTANT_APPLICATION;
+
+		// Apply to target
+		SourceComponent->ApplyGameplayEffectToTarget(BaseDmgEffect, DestComponent, 1.f);
+
+		float ExpectedValue = StartHealth - DamageValue;
+		float ActualValue = DestComponent->GetSet<USkillSystemTestAttributeSet>()->Health;
+
+		Test->TestTrue(SKILL_TEST_TEXT("Buff Instant Damage Applied"), (ActualValue == ExpectedValue));
+		SKILL_LOG(Log, TEXT("Final Health: %.2f"), DestComponent->GetSet<USkillSystemTestAttributeSet>()->Health);
+	}
+
+	// Apply Damage
+	{
+		UGameplayEffect * BaseDmgEffect = Cast<UGameplayEffect>(StaticConstructObject(UGameplayEffect::StaticClass(), GetTransientPackage(), FName(TEXT("Damage"))));
+		BaseDmgEffect->Modifiers.SetNum(1);
+		BaseDmgEffect->Modifiers[0].Magnitude.SetValue(DamageValue);
+		BaseDmgEffect->Modifiers[0].ModifierType = EGameplayMod::Attribute;
+		BaseDmgEffect->Modifiers[0].ModifierOp = EGameplayModOp::Additive;
+		BaseDmgEffect->Modifiers[0].Attribute.SetUProperty(DamageProperty);
+		BaseDmgEffect->Modifiers[0].OwnedTags.AddTag(FName(TEXT("Damage.Buffable")));
+		BaseDmgEffect->Duration.Value = UGameplayEffect::INSTANT_APPLICATION;
+
+		// Apply to target
+		DestComponent->ApplyGameplayEffectToTarget(BaseDmgEffect, SourceComponent, 1.f);
+
+		float ExpectedValue = StartHealth - DamageValue;
+		float ActualValue = SourceComponent->GetSet<USkillSystemTestAttributeSet>()->Health;
+
+		Test->TestTrue(SKILL_TEST_TEXT("Instant Damage Applied"), (ActualValue == ExpectedValue));
+		SKILL_LOG(Log, TEXT("Final Health: %.2f"), SourceComponent->GetSet<USkillSystemTestAttributeSet>()->Health);
+	}
+
+	// Apply Damage again to make sure that the buff only applied once
+	{
+		UGameplayEffect * BaseDmgEffect = Cast<UGameplayEffect>(StaticConstructObject(UGameplayEffect::StaticClass(), GetTransientPackage(), FName(TEXT("Damage"))));
+		BaseDmgEffect->Modifiers.SetNum(1);
+		BaseDmgEffect->Modifiers[0].Magnitude.SetValue(DamageValue);
+		BaseDmgEffect->Modifiers[0].ModifierType = EGameplayMod::Attribute;
+		BaseDmgEffect->Modifiers[0].ModifierOp = EGameplayModOp::Additive;
+		BaseDmgEffect->Modifiers[0].Attribute.SetUProperty(DamageProperty);
+		BaseDmgEffect->Modifiers[0].OwnedTags.AddTag(FName(TEXT("Damage.Basic")));
+		BaseDmgEffect->Duration.Value = UGameplayEffect::INSTANT_APPLICATION;
+
+		// Apply to target
+		DestComponent->ApplyGameplayEffectToTarget(BaseDmgEffect, SourceComponent, 1.f);
+
+		float ExpectedValue = StartHealth - DamageValue - (DamageValue / DamageProtectionDivisor);
+		float ActualValue = SourceComponent->GetSet<USkillSystemTestAttributeSet>()->Health;
+
+		Test->TestTrue(SKILL_TEST_TEXT("Instant Damage Applied"), (ActualValue == ExpectedValue));
+		SKILL_LOG(Log, TEXT("Final Health: %.2f"), SourceComponent->GetSet<USkillSystemTestAttributeSet>()->Health);
+	}
+
+	World->EditorDestroyActor(SourceActor, false);
+	World->EditorDestroyActor(DestActor, false);
+
+	return true;
+}
+
 bool GameplayEffectsTest_DurationDamage(UWorld *World, FAutomationTestBase * Test)
 {
 	const float StartHealth = 100.f;
@@ -1387,7 +1713,7 @@ bool GameplayEffectsTest_DurationDamage(UWorld *World, FAutomationTestBase * Tes
 		BaseDmgEffect->Duration.Value = Duration;
 		BaseDmgEffect->Period.Value = UGameplayEffect::NO_PERIOD;
 
-		SourceComponent->ApplyGameplayEffectSpecToTarget(BaseDmgEffect, DestComponent, 1.f);
+		SourceComponent->ApplyGameplayEffectToTarget(BaseDmgEffect, DestComponent, 1.f);
 
 		// The effect should instantly execute one time without ticking (for now at least)
 		float ActualValue = DestComponent->GetSet<USkillSystemTestAttributeSet>()->Health;
@@ -1465,7 +1791,7 @@ bool GameplayEffectsTest_PeriodicDamage(UWorld *World, FAutomationTestBase * Tes
 		BaseDmgEffect->Period.Value = 1.f; // Apply every 1 second
 		BaseDmgEffect->GameplayCues.Add( FGameplayEffectCue(FName(TEXT("GameplayCue.Burning")), 1.f, 10.f) );
 
-		SourceComponent->ApplyGameplayEffectSpecToTarget(BaseDmgEffect, DestComponent, 5.f);
+		SourceComponent->ApplyGameplayEffectToTarget(BaseDmgEffect, DestComponent, 5.f);
 
 		// The effect should execute on the next tick
 		float ActualValue = DestComponent->GetSet<USkillSystemTestAttributeSet>()->Health;
@@ -1554,7 +1880,7 @@ bool GameplayEffectsTest_LifestealExtension(UWorld *World, FAutomationTestBase *
 		BaseDmgEffect->Duration.Value = UGameplayEffect::INSTANT_APPLICATION;
 		BaseDmgEffect->Period.Value = UGameplayEffect::NO_PERIOD;
 
-		SourceComponent->ApplyGameplayEffectSpecToTarget(BaseDmgEffect, DestComponent, 5.f);
+		SourceComponent->ApplyGameplayEffectToTarget(BaseDmgEffect, DestComponent, 5.f);
 
 		// The effect should instantly execute one time without ticking (for now at least)
 		{
@@ -1635,7 +1961,7 @@ bool GameplayEffectsTest_ShieldExtension(UWorld *World, FAutomationTestBase * Te
 		
 		// Apply 1
 		{
-			SourceComponent->ApplyGameplayEffectSpecToTarget(BaseDmgEffect, DestComponent, 5.f);
+			SourceComponent->ApplyGameplayEffectToTarget(BaseDmgEffect, DestComponent, 5.f);
 
 			// Health should be the same
 			{
@@ -1654,7 +1980,7 @@ bool GameplayEffectsTest_ShieldExtension(UWorld *World, FAutomationTestBase * Te
 
 		// Apply 2
 		{
-			SourceComponent->ApplyGameplayEffectSpecToTarget(BaseDmgEffect, DestComponent, 5.f);
+			SourceComponent->ApplyGameplayEffectToTarget(BaseDmgEffect, DestComponent, 5.f);
 
 			// Health should be the same
 			{
@@ -1673,7 +1999,7 @@ bool GameplayEffectsTest_ShieldExtension(UWorld *World, FAutomationTestBase * Te
 		// Apply 3
 		{
 			// Now we lose health
-			SourceComponent->ApplyGameplayEffectSpecToTarget(BaseDmgEffect, DestComponent, 5.f);
+			SourceComponent->ApplyGameplayEffectToTarget(BaseDmgEffect, DestComponent, 5.f);
 
 			// Now we should have lost some health
 			float ActualValue = DestComponent->GetSet<USkillSystemTestAttributeSet>()->Health;
@@ -1764,7 +2090,7 @@ bool GameplayEffectsTest_ShieldExtensionMultiple(UWorld *World, FAutomationTestB
 
 		// Apply small damage
 		{
-			SourceComponent->ApplyGameplayEffectSpecToTarget(SmallDmgEffect, DestComponent, 5.f);
+			SourceComponent->ApplyGameplayEffectToTarget(SmallDmgEffect, DestComponent, 5.f);
 
 			// Health should be the same
 			{
@@ -1790,7 +2116,7 @@ bool GameplayEffectsTest_ShieldExtensionMultiple(UWorld *World, FAutomationTestB
 
 		// Apply large damage
 		{
-			SourceComponent->ApplyGameplayEffectSpecToTarget(LargeDmgEffect, DestComponent, 5.f);
+			SourceComponent->ApplyGameplayEffectToTarget(LargeDmgEffect, DestComponent, 5.f);
 
 			// Health should still be the same
 			{
@@ -1818,7 +2144,7 @@ bool GameplayEffectsTest_ShieldExtensionMultiple(UWorld *World, FAutomationTestB
 		// Apply large damage again
 		{
 			// Now we lose health
-			SourceComponent->ApplyGameplayEffectSpecToTarget(LargeDmgEffect, DestComponent, 5.f);
+			SourceComponent->ApplyGameplayEffectToTarget(LargeDmgEffect, DestComponent, 5.f);
 
 			float HealthDelta = ShieldAmount + ShieldAmount + DamageValueSmall + DamageValueLarge + DamageValueLarge;
 
@@ -1951,7 +2277,7 @@ bool GameplayEffectsTest_InstantDamage_ScalingExplicit(UWorld *World, FAutomatio
 		BaseDmgEffect->Modifiers[0].OwnedTags.AddTag(FName(TEXT("Damage.Basic")));
 		BaseDmgEffect->Duration.Value = UGameplayEffect::INSTANT_APPLICATION;
 
-		SourceComponent->ApplyGameplayEffectSpecToTarget(BaseDmgEffect, DestComponent, LevelOfDamage);
+		SourceComponent->ApplyGameplayEffectToTarget(BaseDmgEffect, DestComponent, LevelOfDamage);
 
 		float ActualValue = DestComponent->GetSet<USkillSystemTestAttributeSet>()->Health;
 		float ExpectedValue = StartHealth - (LevelOfDamage * SourceDamageScale);
@@ -1997,7 +2323,7 @@ bool GameplayEffectsTest_InstantDamage_ScalingGlobal(UWorld *World, FAutomationT
 		BaseDmgEffect->Modifiers[0].OwnedTags.AddTag(FName(TEXT("Damage.Basic")));
 		BaseDmgEffect->Duration.Value = UGameplayEffect::INSTANT_APPLICATION;
 
-		SourceComponent->ApplyGameplayEffectSpecToTarget(BaseDmgEffect, DestComponent, LevelOfDamage);
+		SourceComponent->ApplyGameplayEffectToTarget(BaseDmgEffect, DestComponent, LevelOfDamage);
 
 		float ActualValue = DestComponent->GetSet<USkillSystemTestAttributeSet>()->Health;
 		float ExpectedValue = StartHealth - (LevelOfDamage);
@@ -2012,7 +2338,7 @@ bool GameplayEffectsTest_InstantDamage_ScalingGlobal(UWorld *World, FAutomationT
 	return true;
 }
 
-bool GameplayEffectsTest_InstantDamage_Override(UWorld *World, FAutomationTestBase * Test)
+bool GameplayEffectsTest_InstantDamage_OverrideScaling(UWorld *World, FAutomationTestBase * Test)
 {
 	// This example overrides global scaling. The setup is the same as GameplayEffectsTest_InstantDamage_ScalingGlobal except now the source
 	// has an explicit override table that will take precedent of the global table.
@@ -2049,7 +2375,7 @@ bool GameplayEffectsTest_InstantDamage_Override(UWorld *World, FAutomationTestBa
 		BaseDmgEffect->Modifiers[0].OwnedTags.AddTag(FName(TEXT("Damage.Basic")));
 		BaseDmgEffect->Duration.Value = UGameplayEffect::INSTANT_APPLICATION;
 
-		SourceComponent->ApplyGameplayEffectSpecToTarget(BaseDmgEffect, DestComponent, LevelOfDamage);
+		SourceComponent->ApplyGameplayEffectToTarget(BaseDmgEffect, DestComponent, LevelOfDamage);
 
 		float ActualValue = DestComponent->GetSet<USkillSystemTestAttributeSet>()->Health;
 		float ExpectedValue = StartHealth - (LevelOfDamage * SourceDamageScale);
@@ -2096,7 +2422,7 @@ bool GameplayEffectsTest_InstantDamageRequiredTag(UWorld *World, FAutomationTest
 		BaseProtectEffect->GameplayEffectRequiredTags.AddTag(FName(TEXT("Damage.Type2")));
 
 		// Apply to self
-		DestComponent->ApplyGameplayEffectSpecToTarget(BaseProtectEffect, DestComponent, 1.f);
+		DestComponent->ApplyGameplayEffectToTarget(BaseProtectEffect, DestComponent, 1.f);
 	}
 
 	// Apply Damage
@@ -2114,7 +2440,7 @@ bool GameplayEffectsTest_InstantDamageRequiredTag(UWorld *World, FAutomationTest
 		BaseDmgEffect->GameplayEffectTags.AddTag(FName(TEXT("Damage.Type1")));
 
 		// Apply to target
-		SourceComponent->ApplyGameplayEffectSpecToTarget(BaseDmgEffect, DestComponent, 1.f);
+		SourceComponent->ApplyGameplayEffectToTarget(BaseDmgEffect, DestComponent, 1.f);
 
 		float ExpectedValue = (StartHealth + DamageValue);
 		
@@ -2140,7 +2466,7 @@ bool GameplayEffectsTest_InstantDamageRequiredTag(UWorld *World, FAutomationTest
 		BaseDmgEffect->GameplayEffectTags.AddTag(FName(TEXT("Damage.Type2")));
 
 		// Apply to target
-		SourceComponent->ApplyGameplayEffectSpecToTarget(BaseDmgEffect, DestComponent, 1.f);
+		SourceComponent->ApplyGameplayEffectToTarget(BaseDmgEffect, DestComponent, 1.f);
 
 		float ExpectedValue = (StartHealth + (DamageValue / DamageProtectionDivisor));
 		float ActualValue = DestComponent->GetSet<USkillSystemTestAttributeSet>()->Health;
@@ -2186,7 +2512,7 @@ bool GameplayEffectsTest_InstantDamageIgnoreTag(UWorld *World, FAutomationTestBa
 		BaseProtectEffect->GameplayEffectIgnoreTags.AddTag(FName(TEXT("Damage.Type1")));
 
 		// Apply to self
-		DestComponent->ApplyGameplayEffectSpecToTarget(BaseProtectEffect, DestComponent, 1.f);
+		DestComponent->ApplyGameplayEffectToTarget(BaseProtectEffect, DestComponent, 1.f);
 	}
 
 	// Apply Damage
@@ -2204,7 +2530,7 @@ bool GameplayEffectsTest_InstantDamageIgnoreTag(UWorld *World, FAutomationTestBa
 		BaseDmgEffect->GameplayEffectTags.AddTag(FName(TEXT("Damage.Type1")));
 
 		// Apply to target
-		SourceComponent->ApplyGameplayEffectSpecToTarget(BaseDmgEffect, DestComponent, 1.f);
+		SourceComponent->ApplyGameplayEffectToTarget(BaseDmgEffect, DestComponent, 1.f);
 
 		float ExpectedValue = (StartHealth + DamageValue);
 
@@ -2230,7 +2556,7 @@ bool GameplayEffectsTest_InstantDamageIgnoreTag(UWorld *World, FAutomationTestBa
 		BaseDmgEffect->GameplayEffectTags.AddTag(FName(TEXT("Damage.Type2")));
 
 		// Apply to target
-		SourceComponent->ApplyGameplayEffectSpecToTarget(BaseDmgEffect, DestComponent, 1.f);
+		SourceComponent->ApplyGameplayEffectToTarget(BaseDmgEffect, DestComponent, 1.f);
 
 		float ExpectedValue = (StartHealth + (DamageValue / DamageProtectionDivisor));
 
@@ -2275,7 +2601,7 @@ bool GameplayEffectsTest_InstantDamageModifierPassesTag(UWorld *World, FAutomati
 		BaseDmgEffect->Duration.SetValue(UGameplayEffect::INFINITE_DURATION);
 
 		// Apply to self
-		SourceComponent->ApplyGameplayEffectSpecToTarget(BaseDmgEffect, SourceComponent, 1.f);
+		SourceComponent->ApplyGameplayEffectToTarget(BaseDmgEffect, SourceComponent, 1.f);
 	}
 
 	// Setup a GE to modify IncomingGEs
@@ -2294,7 +2620,7 @@ bool GameplayEffectsTest_InstantDamageModifierPassesTag(UWorld *World, FAutomati
 		BaseProtectEffect->CopyPolicy = EGameplayEffectCopyPolicy::AlwaysLink;
 
 		// Apply to self
-		DestComponent->ApplyGameplayEffectSpecToTarget(BaseProtectEffect, DestComponent, 1.f);
+		DestComponent->ApplyGameplayEffectToTarget(BaseProtectEffect, DestComponent, 1.f);
 	}
 
 	// Apply Damage
@@ -2311,7 +2637,7 @@ bool GameplayEffectsTest_InstantDamageModifierPassesTag(UWorld *World, FAutomati
 		BaseDmgEffect->Duration.Value = UGameplayEffect::INSTANT_APPLICATION;
 
 		// Apply to target
-		SourceComponent->ApplyGameplayEffectSpecToTarget(BaseDmgEffect, DestComponent, 1.f);
+		SourceComponent->ApplyGameplayEffectToTarget(BaseDmgEffect, DestComponent, 1.f);
 
 		float ExpectedValue = (StartHealth + ((DamageValue * BonusDamageMultiplier) / DamageProtectionDivisor));
 		float ActualValue = DestComponent->GetSet<USkillSystemTestAttributeSet>()->Health;
@@ -2357,7 +2683,7 @@ bool GameplayEffectsTest_InstantDamageModifierTag(UWorld *World, FAutomationTest
 		BaseDmgEffect->Duration.SetValue(UGameplayEffect::INFINITE_DURATION);
 
 		// Apply to self
-		SourceComponent->ApplyGameplayEffectSpecToTarget(BaseDmgEffect, SourceComponent, 1.f);
+		SourceComponent->ApplyGameplayEffectToTarget(BaseDmgEffect, SourceComponent, 1.f);
 	}
 
 	// Setup a GE to modify IncomingGEs
@@ -2376,7 +2702,7 @@ bool GameplayEffectsTest_InstantDamageModifierTag(UWorld *World, FAutomationTest
 		BaseProtectEffect->CopyPolicy = EGameplayEffectCopyPolicy::AlwaysLink;
 
 		// Apply to self
-		DestComponent->ApplyGameplayEffectSpecToTarget(BaseProtectEffect, DestComponent, 1.f);
+		DestComponent->ApplyGameplayEffectToTarget(BaseProtectEffect, DestComponent, 1.f);
 	}
 
 	// Apply Damage
@@ -2393,7 +2719,7 @@ bool GameplayEffectsTest_InstantDamageModifierTag(UWorld *World, FAutomationTest
 		BaseDmgEffect->Duration.Value = UGameplayEffect::INSTANT_APPLICATION;
 
 		// Apply to target
-		SourceComponent->ApplyGameplayEffectSpecToTarget(BaseDmgEffect, DestComponent, 1.f);
+		SourceComponent->ApplyGameplayEffectToTarget(BaseDmgEffect, DestComponent, 1.f);
 
 		float ExpectedValue = (StartHealth + ((DamageValue + BonusDamageValue) / DamageProtectionDivisor));
 		float ActualValue = DestComponent->GetSet<USkillSystemTestAttributeSet>()->Health;
@@ -2449,7 +2775,7 @@ bool GameplayEffectsTest_InstantDamage_ScalingProperty(UWorld *World, FAutomatio
 		BaseDmgEffect->LevelInfo.InheritLevelFromOwner = false;
 		
 
-		SourceComponent->ApplyGameplayEffectSpecToTarget(BaseDmgEffect, DestComponent);
+		SourceComponent->ApplyGameplayEffectToTarget(BaseDmgEffect, DestComponent);
 
 		float ActualValue = DestComponent->GetSet<USkillSystemTestAttributeSet>()->Health;
 		float ExpectedValue = StartHealth - (PhysicalDamage * GameplayEffectScaling);
@@ -2504,7 +2830,7 @@ bool GameplayEffectsTest_InstantDamage_ScalingPropertyNested(UWorld *World, FAut
 		BaseDmgEffect->Modifiers[0].LevelInfo.InheritLevelFromOwner = false;
 		BaseDmgEffect->Duration.Value = UGameplayEffect::INSTANT_APPLICATION;
 
-		SourceComponent->ApplyGameplayEffectSpecToTarget(BaseDmgEffect, DestComponent);
+		SourceComponent->ApplyGameplayEffectToTarget(BaseDmgEffect, DestComponent);
 
 		float ActualValue = DestComponent->GetSet<USkillSystemTestAttributeSet>()->Health;
 		float ExpectedValue = StartHealth - (PhysicalDamage * GameplayEffectScaling);
@@ -2563,7 +2889,7 @@ bool GameplayEffectsTest_DotDamage_ScalingProperty_Snapshot(UWorld *World, FAuto
 		BaseDmgEffect->LevelInfo.InheritLevelFromOwner = false;
 		BaseDmgEffect->LevelInfo.TakeSnapshotOnInit = true;								// But just a snapshot of their SpellDamage when we are applied
 
-		SourceComponent->ApplyGameplayEffectSpecToTarget(BaseDmgEffect, DestComponent);
+		SourceComponent->ApplyGameplayEffectToTarget(BaseDmgEffect, DestComponent);
 	}
 
 
@@ -2578,7 +2904,7 @@ bool GameplayEffectsTest_DotDamage_ScalingProperty_Snapshot(UWorld *World, FAuto
 		SpellDmgEffect->Modifiers[0].OwnedTags.AddTag(FName(TEXT("SpellDmg.Buff")));
 		SpellDmgEffect->Duration.Value = UGameplayEffect::INSTANT_APPLICATION;
 
-		SourceComponent->ApplyGameplayEffectSpecToTarget(SpellDmgEffect, SourceComponent);
+		SourceComponent->ApplyGameplayEffectToTarget(SpellDmgEffect, SourceComponent);
 	}
 
 	{
@@ -2642,7 +2968,7 @@ bool GameplayEffectsTest_DotDamage_ScalingProperty_Dynamic(UWorld *World, FAutom
 		BaseDmgEffect->LevelInfo.InheritLevelFromOwner = false;
 		BaseDmgEffect->LevelInfo.TakeSnapshotOnInit = false;							// But level is dynamic, if SpellDamage changers after we apply, we update.
 
-		SourceComponent->ApplyGameplayEffectSpecToTarget(BaseDmgEffect, DestComponent);
+		SourceComponent->ApplyGameplayEffectToTarget(BaseDmgEffect, DestComponent);
 
 		float SpellDamageTest = SourceComponent->GetSet<USkillSystemTestAttributeSet>()->SpellDamage;
 		check(SpellDamageTest == SpellDamage);
@@ -2668,7 +2994,7 @@ bool GameplayEffectsTest_DotDamage_ScalingProperty_Dynamic(UWorld *World, FAutom
 		SpellDmgEffect->Modifiers[0].OwnedTags.AddTag(FName(TEXT("SpellDmg.Buff")));
 		SpellDmgEffect->Duration.Value = UGameplayEffect::INSTANT_APPLICATION;
 
-		SourceComponent->ApplyGameplayEffectSpecToTarget(SpellDmgEffect, SourceComponent);
+		SourceComponent->ApplyGameplayEffectToTarget(SpellDmgEffect, SourceComponent);
 
 		float ActualValue = SourceComponent->GetSet<USkillSystemTestAttributeSet>()->SpellDamage;
 		float ExpectedValue = SpellDamage2;
@@ -2732,7 +3058,7 @@ bool GameplayEffectsTest_MetaAttributes(UWorld *World, FAutomationTestBase * Tes
 		StrengthMaxHealhEffect->LevelInfo.InheritLevelFromOwner = false;
 		StrengthMaxHealhEffect->LevelInfo.TakeSnapshotOnInit = false;						// But level is dynamic, if SpellDamage changers after we apply, we update.
 
-		SourceComponent->ApplyGameplayEffectSpecToTarget(StrengthMaxHealhEffect, SourceComponent);
+		SourceComponent->ApplyGameplayEffectToTarget(StrengthMaxHealhEffect, SourceComponent);
 
 		// Strength starts at 0, so confirm it did nothing yet.
 		float ActualValue = SourceComponent->GetSet<USkillSystemTestAttributeSet>()->MaxHealth;
@@ -2751,7 +3077,7 @@ bool GameplayEffectsTest_MetaAttributes(UWorld *World, FAutomationTestBase * Tes
 		StrEffect->Modifiers[0].OwnedTags.AddTag(FName(TEXT("SpellDmg.Buff")));
 		StrEffect->Duration.Value = UGameplayEffect::INSTANT_APPLICATION;
 
-		SourceComponent->ApplyGameplayEffectSpecToTarget(StrEffect, SourceComponent);
+		SourceComponent->ApplyGameplayEffectToTarget(StrEffect, SourceComponent);
 
 		{
 			float ActualValue = SourceComponent->GetSet<USkillSystemTestAttributeSet>()->Strength;
@@ -2808,7 +3134,7 @@ bool GameplayEffectsTest_TagOrdering(UWorld *World, FAutomationTestBase * Test)
 		BaseDmgEffect->GameplayEffectRequiredTags.AddTag(FName(TEXT("Damage")));
 
 		// Apply to self
-		SourceComponent->ApplyGameplayEffectSpecToTarget(BaseDmgEffect, SourceComponent, 1.f);
+		SourceComponent->ApplyGameplayEffectToTarget(BaseDmgEffect, SourceComponent, 1.f);
 	}
 
 	{
@@ -2828,7 +3154,7 @@ bool GameplayEffectsTest_TagOrdering(UWorld *World, FAutomationTestBase * Test)
 		BaseDmgEffect->GameplayEffectRequiredTags.AddTag(FName(TEXT("Damage")));
 
 		// Apply to self
-		SourceComponent->ApplyGameplayEffectSpecToTarget(BaseDmgEffect, SourceComponent, 1.f);
+		SourceComponent->ApplyGameplayEffectToTarget(BaseDmgEffect, SourceComponent, 1.f);
 	}
 
 	// Apply Damage
@@ -2847,7 +3173,7 @@ bool GameplayEffectsTest_TagOrdering(UWorld *World, FAutomationTestBase * Test)
 		BaseDmgEffect->GameplayEffectTags.AddTag(FName(TEXT("Damage")));
 
 		// Apply to target
-		SourceComponent->ApplyGameplayEffectSpecToTarget(BaseDmgEffect, DestComponent, 1.f);
+		SourceComponent->ApplyGameplayEffectToTarget(BaseDmgEffect, DestComponent, 1.f);
 
 		float ExpectedValue = (StartHealth - (DamageValue * BonusDamageMultiplier));
 		float ActualValue = DestComponent->GetSet<USkillSystemTestAttributeSet>()->Health;
@@ -2890,7 +3216,7 @@ bool GameplayEffectsTest_StackingHighest(UWorld *World, FAutomationTestBase * Te
 		BaseStackedEffect->Period.Value = 1.f;
 		BaseStackedEffect->StackingPolicy = EGameplayEffectStackingPolicy::Highest;
 
-		SourceComponent->ApplyGameplayEffectSpecToTarget(BaseStackedEffect, DestComponent);
+		SourceComponent->ApplyGameplayEffectToTarget(BaseStackedEffect, DestComponent);
 	}
 
 	{
@@ -2915,7 +3241,7 @@ bool GameplayEffectsTest_StackingHighest(UWorld *World, FAutomationTestBase * Te
 		BaseStackedEffect->Period.Value = 1.f;
 		BaseStackedEffect->StackingPolicy = EGameplayEffectStackingPolicy::Highest;
 
-		SourceComponent->ApplyGameplayEffectSpecToTarget(BaseStackedEffect, DestComponent);
+		SourceComponent->ApplyGameplayEffectToTarget(BaseStackedEffect, DestComponent);
 	}
 
 	{
@@ -2965,7 +3291,7 @@ bool GameplayEffectsTest_StackingLowest(UWorld *World, FAutomationTestBase * Tes
 		BaseStackedEffect->Period.Value = 1.f;
 		BaseStackedEffect->StackingPolicy = EGameplayEffectStackingPolicy::Lowest;
 
-		SourceComponent->ApplyGameplayEffectSpecToTarget(BaseStackedEffect, DestComponent);
+		SourceComponent->ApplyGameplayEffectToTarget(BaseStackedEffect, DestComponent);
 	}
 
 	{
@@ -2990,7 +3316,7 @@ bool GameplayEffectsTest_StackingLowest(UWorld *World, FAutomationTestBase * Tes
 		BaseStackedEffect->Period.Value = 1.f;
 		BaseStackedEffect->StackingPolicy = EGameplayEffectStackingPolicy::Lowest;
 
-		SourceComponent->ApplyGameplayEffectSpecToTarget(BaseStackedEffect, DestComponent);
+		SourceComponent->ApplyGameplayEffectToTarget(BaseStackedEffect, DestComponent);
 	}
 
 	{
@@ -3040,7 +3366,7 @@ bool GameplayEffectsTest_StackingUnlimited(UWorld *World, FAutomationTestBase *
 		BaseStackedEffect->Period.Value = 1.f;
 		BaseStackedEffect->StackingPolicy = EGameplayEffectStackingPolicy::Unlimited;
 
-		SourceComponent->ApplyGameplayEffectSpecToTarget(BaseStackedEffect, DestComponent);
+		SourceComponent->ApplyGameplayEffectToTarget(BaseStackedEffect, DestComponent);
 	}
 
 	{
@@ -3056,7 +3382,7 @@ bool GameplayEffectsTest_StackingUnlimited(UWorld *World, FAutomationTestBase *
 		BaseStackedEffect->Period.Value = 1.f;
 		BaseStackedEffect->StackingPolicy = EGameplayEffectStackingPolicy::Unlimited;
 
-		SourceComponent->ApplyGameplayEffectSpecToTarget(BaseStackedEffect, DestComponent);
+		SourceComponent->ApplyGameplayEffectToTarget(BaseStackedEffect, DestComponent);
 	}
 
 	{
@@ -3106,7 +3432,7 @@ bool GameplayEffectsTest_StackingRemoval(UWorld *World, FAutomationTestBase * Te
 		BaseStackedEffect->Period.Value = 1.f;
 		BaseStackedEffect->StackingPolicy = EGameplayEffectStackingPolicy::Highest;
 
-		SourceComponent->ApplyGameplayEffectSpecToTarget(BaseStackedEffect, DestComponent);
+		SourceComponent->ApplyGameplayEffectToTarget(BaseStackedEffect, DestComponent);
 	}
 
 	{
@@ -3122,7 +3448,7 @@ bool GameplayEffectsTest_StackingRemoval(UWorld *World, FAutomationTestBase * Te
 		BaseStackedEffect->Period.Value = 1.f;
 		BaseStackedEffect->StackingPolicy = EGameplayEffectStackingPolicy::Highest;
 
-		SourceComponent->ApplyGameplayEffectSpecToTarget(BaseStackedEffect, DestComponent);
+		SourceComponent->ApplyGameplayEffectToTarget(BaseStackedEffect, DestComponent);
 	}
 
 	{
@@ -3179,7 +3505,7 @@ bool GameplayEffectsTest_StackingNoStack(UWorld *World, FAutomationTestBase * Te
 		UnstackableEffect->Period.Value = 1.f;
 		UnstackableEffect->StackingPolicy = EGameplayEffectStackingPolicy::Highest;
 
-		SourceComponent->ApplyGameplayEffectSpecToTarget(UnstackableEffect, DestComponent);
+		SourceComponent->ApplyGameplayEffectToTarget(UnstackableEffect, DestComponent);
 	}
 
 	{
@@ -3195,7 +3521,7 @@ bool GameplayEffectsTest_StackingNoStack(UWorld *World, FAutomationTestBase * Te
 		UnstackableEffect->Period.Value = 1.f;
 		UnstackableEffect->StackingPolicy = EGameplayEffectStackingPolicy::Highest;
 
-		SourceComponent->ApplyGameplayEffectSpecToTarget(UnstackableEffect, DestComponent);
+		SourceComponent->ApplyGameplayEffectToTarget(UnstackableEffect, DestComponent);
 	}
 
 	{
@@ -3248,7 +3574,7 @@ bool GameplayEffectsTest_StackingCustomCapped(UWorld *World, FAutomationTestBase
 		BaseStackedEffect->StackingExtension = UGameplayEffectStackingExtension_CappedNumberTest::StaticClass();
 //		BaseStackedEffect->StackingExtension = UGameplayEffectStackingExtension_DiminishingReturnsTest::StaticClass();
 
-		SourceComponent->ApplyGameplayEffectSpecToTarget(BaseStackedEffect, DestComponent);
+		SourceComponent->ApplyGameplayEffectToTarget(BaseStackedEffect, DestComponent);
 	}
 
 	{
@@ -3266,7 +3592,7 @@ bool GameplayEffectsTest_StackingCustomCapped(UWorld *World, FAutomationTestBase
 		BaseStackedEffect->StackingPolicy = EGameplayEffectStackingPolicy::Callback;
 		BaseStackedEffect->StackingExtension = UGameplayEffectStackingExtension_CappedNumberTest::StaticClass();
 
-		SourceComponent->ApplyGameplayEffectSpecToTarget(BaseStackedEffect, DestComponent);
+		SourceComponent->ApplyGameplayEffectToTarget(BaseStackedEffect, DestComponent);
 	}
 
 	{
@@ -3284,7 +3610,7 @@ bool GameplayEffectsTest_StackingCustomCapped(UWorld *World, FAutomationTestBase
 		BaseStackedEffect->StackingPolicy = EGameplayEffectStackingPolicy::Callback;
 		BaseStackedEffect->StackingExtension = UGameplayEffectStackingExtension_CappedNumberTest::StaticClass();
 
-		SourceComponent->ApplyGameplayEffectSpecToTarget(BaseStackedEffect, DestComponent);
+		SourceComponent->ApplyGameplayEffectToTarget(BaseStackedEffect, DestComponent);
 	}
 
 	{
@@ -3311,7 +3637,7 @@ bool GameplayEffectsTest_StackingCustomCapped(UWorld *World, FAutomationTestBase
 		BaseStackedEffect->StackingPolicy = EGameplayEffectStackingPolicy::Callback;
 		BaseStackedEffect->StackingExtension = UGameplayEffectStackingExtension_CappedNumberTest::StaticClass();
 
-		SourceComponent->ApplyGameplayEffectSpecToTarget(BaseStackedEffect, DestComponent);
+		SourceComponent->ApplyGameplayEffectToTarget(BaseStackedEffect, DestComponent);
 	}
 
 	{
@@ -3381,7 +3707,7 @@ bool GameplayEffectsTest_StackingCustomDiminishingReturns(UWorld *World, FAutoma
 		BaseStackedEffect->StackingPolicy = EGameplayEffectStackingPolicy::Callback;
 		BaseStackedEffect->StackingExtension = UGameplayEffectStackingExtension_DiminishingReturnsTest::StaticClass();
 
-		SourceComponent->ApplyGameplayEffectSpecToTarget(BaseStackedEffect, DestComponent);
+		SourceComponent->ApplyGameplayEffectToTarget(BaseStackedEffect, DestComponent);
 	}
 
 	{
@@ -3408,7 +3734,7 @@ bool GameplayEffectsTest_StackingCustomDiminishingReturns(UWorld *World, FAutoma
 		BaseStackedEffect->StackingPolicy = EGameplayEffectStackingPolicy::Callback;
 		BaseStackedEffect->StackingExtension = UGameplayEffectStackingExtension_DiminishingReturnsTest::StaticClass();
 
-		SourceComponent->ApplyGameplayEffectSpecToTarget(BaseStackedEffect, DestComponent);
+		SourceComponent->ApplyGameplayEffectToTarget(BaseStackedEffect, DestComponent);
 	}
 
 	{
@@ -3435,7 +3761,7 @@ bool GameplayEffectsTest_StackingCustomDiminishingReturns(UWorld *World, FAutoma
 		BaseStackedEffect->StackingPolicy = EGameplayEffectStackingPolicy::Callback;
 		BaseStackedEffect->StackingExtension = UGameplayEffectStackingExtension_DiminishingReturnsTest::StaticClass();
 
-		SourceComponent->ApplyGameplayEffectSpecToTarget(BaseStackedEffect, DestComponent);
+		SourceComponent->ApplyGameplayEffectToTarget(BaseStackedEffect, DestComponent);
 	}
 
 	{
@@ -3462,7 +3788,7 @@ bool GameplayEffectsTest_StackingCustomDiminishingReturns(UWorld *World, FAutoma
 		BaseStackedEffect->StackingPolicy = EGameplayEffectStackingPolicy::Callback;
 		BaseStackedEffect->StackingExtension = UGameplayEffectStackingExtension_DiminishingReturnsTest::StaticClass();
 
-		SourceComponent->ApplyGameplayEffectSpecToTarget(BaseStackedEffect, DestComponent);
+		SourceComponent->ApplyGameplayEffectToTarget(BaseStackedEffect, DestComponent);
 	}
 
 	{
@@ -3512,7 +3838,7 @@ bool GameplayEffectsTest_StackingDifferentRules(UWorld *World, FAutomationTestBa
 		BaseStackedEffect->Period.Value = 1.f;
 		BaseStackedEffect->StackingPolicy = EGameplayEffectStackingPolicy::Highest;
 
-		SourceComponent->ApplyGameplayEffectSpecToTarget(BaseStackedEffect, DestComponent);
+		SourceComponent->ApplyGameplayEffectToTarget(BaseStackedEffect, DestComponent);
 	}
 
 	{
@@ -3528,7 +3854,7 @@ bool GameplayEffectsTest_StackingDifferentRules(UWorld *World, FAutomationTestBa
 		BaseStackedEffect->Period.Value = 1.f;
 		BaseStackedEffect->StackingPolicy = EGameplayEffectStackingPolicy::Lowest;
 
-		SourceComponent->ApplyGameplayEffectSpecToTarget(BaseStackedEffect, DestComponent);
+		SourceComponent->ApplyGameplayEffectToTarget(BaseStackedEffect, DestComponent);
 	}
 
 	{
@@ -3579,7 +3905,7 @@ bool GameplayEffectsTest_StackingDifferentAttributes(UWorld *World, FAutomationT
 		BaseStackedEffect->Period.Value = 1.f;
 		BaseStackedEffect->StackingPolicy = EGameplayEffectStackingPolicy::Highest;
 
-		SourceComponent->ApplyGameplayEffectSpecToTarget(BaseStackedEffect, DestComponent);
+		SourceComponent->ApplyGameplayEffectToTarget(BaseStackedEffect, DestComponent);
 	}
 
 	{
@@ -3595,7 +3921,7 @@ bool GameplayEffectsTest_StackingDifferentAttributes(UWorld *World, FAutomationT
 		BaseStackedEffect->Period.Value = 1.f;
 		BaseStackedEffect->StackingPolicy = EGameplayEffectStackingPolicy::Highest;
 
-		SourceComponent->ApplyGameplayEffectSpecToTarget(BaseStackedEffect, DestComponent);
+		SourceComponent->ApplyGameplayEffectToTarget(BaseStackedEffect, DestComponent);
 	}
 
 	{
@@ -3650,7 +3976,7 @@ bool GameplayEffectsTest_StackingCustomTwoRules(UWorld *World, FAutomationTestBa
 		BaseStackedEffect->StackingPolicy = EGameplayEffectStackingPolicy::Callback;
 		BaseStackedEffect->StackingExtension = UGameplayEffectStackingExtension_CappedNumberTest::StaticClass();
 
-		SourceComponent->ApplyGameplayEffectSpecToTarget(BaseStackedEffect, DestComponent);
+		SourceComponent->ApplyGameplayEffectToTarget(BaseStackedEffect, DestComponent);
 	}
 
 	{
@@ -3668,7 +3994,7 @@ bool GameplayEffectsTest_StackingCustomTwoRules(UWorld *World, FAutomationTestBa
 		BaseStackedEffect->StackingPolicy = EGameplayEffectStackingPolicy::Callback;
 		BaseStackedEffect->StackingExtension = UGameplayEffectStackingExtension_DiminishingReturnsTest::StaticClass();
 
-		SourceComponent->ApplyGameplayEffectSpecToTarget(BaseStackedEffect, DestComponent);
+		SourceComponent->ApplyGameplayEffectToTarget(BaseStackedEffect, DestComponent);
 	}
 
 	{
@@ -3686,7 +4012,7 @@ bool GameplayEffectsTest_StackingCustomTwoRules(UWorld *World, FAutomationTestBa
 		BaseStackedEffect->StackingPolicy = EGameplayEffectStackingPolicy::Callback;
 		BaseStackedEffect->StackingExtension = UGameplayEffectStackingExtension_DiminishingReturnsTest::StaticClass();
 
-		SourceComponent->ApplyGameplayEffectSpecToTarget(BaseStackedEffect, DestComponent);
+		SourceComponent->ApplyGameplayEffectToTarget(BaseStackedEffect, DestComponent);
 	}
 
 	{
@@ -3739,7 +4065,7 @@ bool GameplayEffectsTest_StackingCustomTwoAttributes(UWorld *World, FAutomationT
 		BaseStackedEffect->StackingPolicy = EGameplayEffectStackingPolicy::Callback;
 		BaseStackedEffect->StackingExtension = UGameplayEffectStackingExtension_CappedNumberTest::StaticClass();
 
-		SourceComponent->ApplyGameplayEffectSpecToTarget(BaseStackedEffect, DestComponent);
+		SourceComponent->ApplyGameplayEffectToTarget(BaseStackedEffect, DestComponent);
 	}
 
 	{
@@ -3757,7 +4083,7 @@ bool GameplayEffectsTest_StackingCustomTwoAttributes(UWorld *World, FAutomationT
 		BaseStackedEffect->StackingPolicy = EGameplayEffectStackingPolicy::Callback;
 		BaseStackedEffect->StackingExtension = UGameplayEffectStackingExtension_CappedNumberTest::StaticClass();
 
-		SourceComponent->ApplyGameplayEffectSpecToTarget(BaseStackedEffect, DestComponent);
+		SourceComponent->ApplyGameplayEffectToTarget(BaseStackedEffect, DestComponent);
 	}
 
 	{
@@ -3775,7 +4101,7 @@ bool GameplayEffectsTest_StackingCustomTwoAttributes(UWorld *World, FAutomationT
 		BaseStackedEffect->StackingPolicy = EGameplayEffectStackingPolicy::Callback;
 		BaseStackedEffect->StackingExtension = UGameplayEffectStackingExtension_CappedNumberTest::StaticClass();
 
-		SourceComponent->ApplyGameplayEffectSpecToTarget(BaseStackedEffect, DestComponent);
+		SourceComponent->ApplyGameplayEffectToTarget(BaseStackedEffect, DestComponent);
 	}
 
 	{
@@ -3831,7 +4157,7 @@ bool GameplayEffectsTest_StackingRemovingModifiers(UWorld *World, FAutomationTes
 		BaseModEffect->GameplayEffectRequiredTags.AddTag(FName(TEXT("Stack")));
 
 		// Apply to self
-		DestComponent->ApplyGameplayEffectSpecToTarget(BaseModEffect, DestComponent, 1.f);
+		DestComponent->ApplyGameplayEffectToTarget(BaseModEffect, DestComponent, 1.f);
 	}
 
 	{
@@ -3848,7 +4174,7 @@ bool GameplayEffectsTest_StackingRemovingModifiers(UWorld *World, FAutomationTes
 		BaseStackedEffect->StackingPolicy = EGameplayEffectStackingPolicy::Highest;
 		BaseStackedEffect->GameplayEffectTags.AddTag(FName(TEXT("Stack")));
 
-		SourceComponent->ApplyGameplayEffectSpecToTarget(BaseStackedEffect, DestComponent);
+		SourceComponent->ApplyGameplayEffectToTarget(BaseStackedEffect, DestComponent);
 	}
 
 	{
@@ -3864,7 +4190,7 @@ bool GameplayEffectsTest_StackingRemovingModifiers(UWorld *World, FAutomationTes
 		BaseStackedEffect->Period.Value = 1.f;
 		BaseStackedEffect->StackingPolicy = EGameplayEffectStackingPolicy::Highest;
 
-		SourceComponent->ApplyGameplayEffectSpecToTarget(BaseStackedEffect, DestComponent);
+		SourceComponent->ApplyGameplayEffectToTarget(BaseStackedEffect, DestComponent);
 	}
 
 	{
@@ -3923,7 +4249,7 @@ bool GameplayEffectsTest_StackingAddingModifiers(UWorld *World, FAutomationTestB
 		BaseStackedEffect->StackingPolicy = EGameplayEffectStackingPolicy::Highest;
 		BaseStackedEffect->GameplayEffectTags.AddTag(FName(TEXT("Stack")));
 
-		SourceComponent->ApplyGameplayEffectSpecToTarget(BaseStackedEffect, DestComponent);
+		SourceComponent->ApplyGameplayEffectToTarget(BaseStackedEffect, DestComponent);
 	}
 
 	{
@@ -3939,7 +4265,7 @@ bool GameplayEffectsTest_StackingAddingModifiers(UWorld *World, FAutomationTestB
 		BaseStackedEffect->Period.Value = 1.f;
 		BaseStackedEffect->StackingPolicy = EGameplayEffectStackingPolicy::Highest;
 
-		SourceComponent->ApplyGameplayEffectSpecToTarget(BaseStackedEffect, DestComponent);
+		SourceComponent->ApplyGameplayEffectToTarget(BaseStackedEffect, DestComponent);
 	}
 
 	{
@@ -3966,7 +4292,7 @@ bool GameplayEffectsTest_StackingAddingModifiers(UWorld *World, FAutomationTestB
 		BaseModEffect->GameplayEffectRequiredTags.AddTag(FName(TEXT("Stack")));
 
 		// Apply to self
-		DestComponent->ApplyGameplayEffectSpecToTarget(BaseModEffect, DestComponent, 1.f);
+		DestComponent->ApplyGameplayEffectToTarget(BaseModEffect, DestComponent, 1.f);
 	}
 
 	{
@@ -4011,23 +4337,38 @@ bool FGameplayEffectsTest::RunTest( const FString& Parameters )
 	GameplayEffectsTest_LinkedBuffDestroy(World, this);
 	GameplayEffectsTest_SnapshotBuffDestroy(World, this);
 	GameplayEffectsTest_DurationBuff(World, this);
+
+	GameplayEffectsTest_DurationDamage(World, this);
+	GameplayEffectsTest_PeriodicDamage(World, this);
 	
+	// Buffing Buffs
 	GameplayEffectsTest_DamageBuffBuff_Basic(World, this);
 	GameplayEffectsTest_DamageBuffBuff_FullLink(World, this);
 	GameplayEffectsTest_DamageBuffBuff_FullSnapshot(World, this);
 	GameplayEffectsTest_DamageBuffBuff_SnapshotLink(World, this);
 
-	GameplayEffectsTest_DurationDamage(World, this);
-	GameplayEffectsTest_PeriodicDamage(World, this);
+	// GameplayEffects that apply other GameplayEffects
+	GameplayEffectsTest_DamageAppliesBuff(World, this);
+	GameplayEffectsTest_BuffAppliesBuff(World, this);
+	GameplayEffectsTest_BuffIndirection(World, this);
 
+	// GameplayEffect extensions
 	GameplayEffectsTest_LifestealExtension(World, this);
 	
 	GameplayEffectsTest_ShieldExtension(World, this);
 	GameplayEffectsTest_ShieldExtensionMultiple(World, this);
 	
+	// Scaling modifiers
 	GameplayEffectsTest_InstantDamage_ScalingExplicit(World, this);
 	GameplayEffectsTest_InstantDamage_ScalingGlobal(World, this);
-	GameplayEffectsTest_InstantDamage_Override(World, this);
+
+	GameplayEffectsTest_InstantDamage_ScalingProperty(World, this);
+	GameplayEffectsTest_InstantDamage_ScalingPropertyNested(World, this);
+
+	GameplayEffectsTest_DotDamage_ScalingProperty_Snapshot(World, this);
+	GameplayEffectsTest_DotDamage_ScalingProperty_Dynamic(World, this);
+
+	GameplayEffectsTest_InstantDamage_OverrideScaling(World, this);
 
 	// Tagging tests
 	GameplayEffectsTest_InstantDamageRequiredTag(World, this);
@@ -4035,11 +4376,6 @@ bool FGameplayEffectsTest::RunTest( const FString& Parameters )
 	
 	GameplayEffectsTest_InstantDamageModifierPassesTag(World, this);
 	GameplayEffectsTest_InstantDamageModifierTag(World, this);
-	GameplayEffectsTest_InstantDamage_ScalingProperty(World, this);
-	GameplayEffectsTest_InstantDamage_ScalingPropertyNested(World, this);
-
-	GameplayEffectsTest_DotDamage_ScalingProperty_Snapshot(World, this);
-	GameplayEffectsTest_DotDamage_ScalingProperty_Dynamic(World, this);	
 
 	GameplayEffectsTest_MetaAttributes(World, this);
 	GameplayEffectsTest_TagOrdering(World, this);
