Gameplay effects may provide immunity and have a probability to apply/execute.

Gameplay effects can provide immunity from other gameplay effects and modifiers to gameplay effects. Immunity can be applied to incoming, outgoing or active gameplay effects independently.
Gameplay effects have a probability of applying to a target or executing on another gameplay effect. For non-instant gameplay effects with no period we apply the probability of affecting another gameplay effect at application time since these effects never execute. Instant and periodic gameplay effects determine if they will affect other gameplay effects at execution time. The probabilty of a gameplay effect applying/executing can be changed via modifiers.
Additional tests have been added to cover probability of applying/executing and immunity.

diff --git a/Engine/Source/Runtime/SkillSystem/Classes/GameplayEffect.h b/Engine/Source/Runtime/SkillSystem/Classes/GameplayEffect.h
--- a/Engine/Source/Runtime/SkillSystem/Classes/GameplayEffect.h
+++ b/Engine/Source/Runtime/SkillSystem/Classes/GameplayEffect.h
@@ -44,7 +44,7 @@ namespace EGameplayModOp
 }
 
 /**
- * Tells us what thing a GameplayEffect modifies.
+ * Tells us what thing a GameplayModifier modifies.
  */
 UENUM(BlueprintType)
 namespace EGameplayMod
@@ -61,6 +61,29 @@ namespace EGameplayMod
 	};
 }
 
+/**
+ * Tells us what thing a GameplayEffect provides immunity to. This must mirror the values in EGameplayMod
+ */
+UENUM(BlueprintType)
+namespace EGameplayImmunity
+{
+	enum Type
+	{
+		None = 0,			// Does not provide immunity
+		OutgoingGE,			// Provides immunity to outgoing GEs
+		IncomingGE,			// Provides immunity to incoming GEs
+		ActiveGE,			// Provides immunity from currently active GEs
+
+		// This must always be at the end
+		Max					UMETA(DisplayName="Invalid")
+	};
+}
+
+checkAtCompileTime(EGameplayMod::OutgoingGE == EGameplayImmunity::OutgoingGE, "EGameplayMod::OutgoingGE and EGameplayImmunity::OutgoingGE must match. Did you forget to modify one of them?");
+checkAtCompileTime(EGameplayMod::IncomingGE == EGameplayImmunity::IncomingGE, "EGameplayMod::IncomingGE and EGameplayImmunity::IncomingGE must match. Did you forget to modify one of them?");
+checkAtCompileTime(EGameplayMod::ActiveGE == EGameplayImmunity::ActiveGE, "EGameplayMod::ActiveGE and EGameplayImmunity::ActiveGE must match. Did you forget to modify one of them?");
+checkAtCompileTime(EGameplayMod::Max == EGameplayImmunity::Max, "EGameplayMod and EGameplayImmunity must cover the same cases. Did you forget to modify one of them?");
+
 /**
  * Tells us what a GameplayEffect modifies when being applied to another GameplayEffect
  */
@@ -72,7 +95,7 @@ namespace EGameplayModEffect
 		Magnitude			= 0x01,		// Modifies magnitude of a GameplayEffect (Always default for Attribute mod)
 		Duration			= 0x02,		// Modifies duration of a GameplayEffect
 		ChanceApplyTarget	= 0x04,		// Modifies chance to apply GameplayEffect to target
-		ChanceApplyEffect	= 0x08,		// Modifies chance to apply GameplayEffect to GameplayEffect
+		ChanceExecuteEffect	= 0x08,		// Modifies chance to execute GameplayEffect on GameplayEffect
 		LinkedGameplayEffect= 0x10,		// Adds a linked GameplayEffect to a GameplayEffect
 
 		// This must always be at the end
@@ -323,12 +346,16 @@ public:
 	FScalableFloat	ChanceToApplyToTarget;
 
 	UPROPERTY(EditDefaultsOnly, Category = Application, meta = (GameplayAttribute = "True"))
-	FScalableFloat	ChanceToApplyToGameplayEffect;
+	FScalableFloat	ChanceToExecuteOnGameplayEffect;
 
 	// other gameplay effects  that will be applied to the target of this effect
 	UPROPERTY(EditDefaultsOnly, BlueprintReadOnly, Category = GameplayEffect)
 	TArray<UGameplayEffect*> TargetEffects;
 
+	// removes or blocks gameplay effects that it applies to if the tags match
+	UPROPERTY(EditDefaultsOnly, BlueprintReadOnly, Category = GameplayEffect)
+	TEnumAsByte<EGameplayImmunity::Type> AppliesImmunityTo;
+
 	// Modify duration of CEs
 
 	// Modify MaxStacks of CEs (maybe... probably not)
@@ -361,13 +388,18 @@ public:
 	UPROPERTY(EditDefaultsOnly, BlueprintReadOnly, Category=Tags)
 	FGameplayTagContainer GameplayEffectIgnoreTags;
 
+	bool AreGameplayEffectTagRequirementsSatisfied(const FGameplayTagContainer& Tags) const
+	{
+		bool bHasRequired = Tags.MatchesAll(GameplayEffectRequiredTags, true);
+		bool bHasIgnored = Tags.MatchesAny(GameplayEffectIgnoreTags, false);
+
+		return bHasRequired && !bHasIgnored;
+	}
+
 	/** Can this GameplayEffect modify the input parameter, based on tags  */
 	bool AreGameplayEffectTagRequirementsSatisfied(const UGameplayEffect *GameplayEffectToBeModified) const
 	{
-		bool HasRequired = GameplayEffectToBeModified->GameplayEffectTags.MatchesAll(GameplayEffectRequiredTags, true);
-		bool HasIgnored = GameplayEffectToBeModified->GameplayEffectTags.MatchesAny(GameplayEffectIgnoreTags, false);
-
-		return HasRequired && !HasIgnored;
+		return AreGameplayEffectTagRequirementsSatisfied(GameplayEffectToBeModified->GameplayEffectTags);
 	}
 
 	// ------------------------------------------------
@@ -685,9 +717,6 @@ struct FAggregatorRef
 	FString ToString() const;
 	void PrintAll() const;
 
-	UPROPERTY()
-	int32	Foo;
-
 private:
 
 	TSharedPtr<struct FAggregator>	SharedPtr;
@@ -827,7 +856,7 @@ struct FAggregator : public TSharedFromThis<FAggregator>
 	virtual ~FAggregator();
 
 	FAggregator & MarkDirty();
-	void CleaerAllDependancies();
+	void ClearAllDependancies();
 
 	const FGameplayModifierEvaluatedData& Evaluate() const;
 
@@ -1079,7 +1108,7 @@ struct FGameplayEffectSpec
 		// If we initialize a GameplayEffectSpec with no level object passed in.
 	}
 
-	FGameplayEffectSpec( const UGameplayEffect *InDef, AActor *Owner, float Level, const FGlobalCurveDataOverride *CurveData );
+	FGameplayEffectSpec( const UGameplayEffect *InDef, AActor *Owner, AActor *Instigator, float Level, const FGlobalCurveDataOverride *CurveData );
 	
 	UPROPERTY()
 	const UGameplayEffect * Def;
@@ -1092,6 +1121,8 @@ struct FGameplayEffectSpec
 
 	float GetDuration() const;
 	float GetPeriod() const;
+	float GetChanceToApplyToTarget() const;
+	float GetChanceToExecuteOnGameplayEffect() const;
 	EGameplayEffectStackingPolicy::Type GetStackingType() const;
 	float GetMagnitude(const FGameplayAttribute &Attribute) const;
 
@@ -1101,8 +1132,15 @@ struct FGameplayEffectSpec
 	UPROPERTY()
 	FAggregatorRef	Duration;
 
+	UPROPERTY()
 	FAggregatorRef	Period;
 
+	UPROPERTY()
+	FAggregatorRef	ChanceToApplyToTarget;
+
+	UPROPERTY()
+	FAggregatorRef	ChanceToExecuteOnGameplayEffect;
+
 	// How this combines with other gameplay effects
 	EGameplayEffectStackingPolicy::Type StackingPolicy;
 	FName StackedAttribName;
@@ -1116,7 +1154,10 @@ struct FGameplayEffectSpec
 
 	void InitModifiers(const FGlobalCurveDataOverride *CurveData, AActor *Owner, float Level);
 
+	// returns the number of modifiers applied to InSpec by the current GameplayEffect Spec
+	// returns -1 if the current GameplayEffectSpec prevents InSpec from being applied
 	int32 ApplyModifiersFrom(FGameplayEffectSpec &InSpec, const FModifierQualifier &QualifierContext);
+
 	int32 ExecuteModifiersFrom(const FGameplayEffectSpec &InSpec, const FModifierQualifier &QualifierContext);
 
 	bool ShouldApplyAsSnapshot(const FModifierQualifier &QualifierContext) const;
@@ -1286,7 +1327,9 @@ struct FActiveGameplayEffectsContainer : public FFastArraySerializer
 	
 	FActiveGameplayEffect & CreateNewActiveGameplayEffect(const FGameplayEffectSpec &Spec);
 
-	void ApplyActiveEffectsTo(OUT FGameplayEffectSpec &Spec, const FModifierQualifier &QualifierContext);
+	// returns true if none of the active effects provide immunity to Spec
+	// returns false if one (or more) of the active effects provides immunity to Spec
+	bool ApplyActiveEffectsTo(OUT FGameplayEffectSpec &Spec, const FModifierQualifier &QualifierContext);
 
 	void ApplySpecToActiveEffectsAndAttributes(FGameplayEffectSpec &Spec, const FModifierQualifier &QualifierContext);
 		
diff --git a/Engine/Source/Runtime/SkillSystem/Private/AttributeComponent.cpp b/Engine/Source/Runtime/SkillSystem/Private/AttributeComponent.cpp
--- a/Engine/Source/Runtime/SkillSystem/Private/AttributeComponent.cpp
+++ b/Engine/Source/Runtime/SkillSystem/Private/AttributeComponent.cpp
@@ -173,11 +173,9 @@ FActiveGameplayEffectHandle UAttributeComponent::ApplyGameplayEffectToTarget(UGa
 {
 	check(GameplayEffect);
 
-	FGameplayEffectSpec	Spec(GameplayEffect, GetOwner(), Level, GetCurveDataOverride());
-	Spec.Def = GameplayEffect;
-	Spec.InstigatorStack.AddInstigator(GetOwner());
+	FGameplayEffectSpec	Spec(GameplayEffect, GetOwner(), GetOwner(), Level, GetCurveDataOverride());
 	
-	return ApplyGameplayEffectSpecToTarget(Spec, Target);
+	return ApplyGameplayEffectSpecToTarget(Spec, Target, BaseQualifier);
 }
 
 /** Helper function since we can't have default/optional values for FModifierQualifier in K2 function */
@@ -189,9 +187,7 @@ FActiveGameplayEffectHandle UAttributeComponent::K2_ApplyGameplayEffectToSelf(co
 /** This is a helper function - it seems like this will be useful as a blueprint interface at the least, but Level parameter may need to be expanded */
 FActiveGameplayEffectHandle UAttributeComponent::ApplyGameplayEffectToSelf(const UGameplayEffect *GameplayEffect, float Level, AActor *Instigator, FModifierQualifier BaseQualifier)
 {
-	FGameplayEffectSpec	Spec(GameplayEffect, GetOwner(), Level, GetCurveDataOverride());
-	Spec.Def = GameplayEffect;
-	Spec.InstigatorStack.AddInstigator(Instigator);
+	FGameplayEffectSpec	Spec(GameplayEffect, GetOwner(), Instigator, Level, GetCurveDataOverride());
 
 	return ApplyGameplayEffectSpecToSelf(Spec, BaseQualifier);
 }
@@ -232,13 +228,24 @@ void UAttributeComponent::TEMP_ApplyActiveGameplayEffects()
 FActiveGameplayEffectHandle UAttributeComponent::ApplyGameplayEffectSpecToTarget(OUT FGameplayEffectSpec &Spec, UAttributeComponent *Target, FModifierQualifier BaseQualifier)
 {
 	// Apply outgoing Effects to the Spec.
-	ActiveGameplayEffects.ApplyActiveEffectsTo(Spec, FModifierQualifier(BaseQualifier).Type(EGameplayMod::OutgoingGE));
+	// Outgoing immunity may stop the outgoing effect from being applied to the target
+	if (ActiveGameplayEffects.ApplyActiveEffectsTo(Spec, FModifierQualifier(BaseQualifier).Type(EGameplayMod::OutgoingGE)))
+	{
+		return Target->ApplyGameplayEffectSpecToSelf(Spec, BaseQualifier);
+	}
 
-	return Target->ApplyGameplayEffectSpecToSelf(Spec, BaseQualifier);
+	return FActiveGameplayEffectHandle();
 }
 
 FActiveGameplayEffectHandle UAttributeComponent::ApplyGameplayEffectSpecToSelf(const FGameplayEffectSpec &Spec, FModifierQualifier BaseQualifier)
 {
+	// check if the effect being applied actually succeeds
+	float ChanceToApply = Spec.GetChanceToApplyToTarget();
+	if ((ChanceToApply < 1.f - SMALL_NUMBER) && (FMath::FRand() > ChanceToApply))
+	{
+		return FActiveGameplayEffectHandle();
+	}
+
 	// Make sure we create our copy of the spec in the right place first...
 	FActiveGameplayEffectHandle	MyHandle;
 	bool bInvokeGameplayCueApplied = false;
@@ -284,7 +291,11 @@ FActiveGameplayEffectHandle UAttributeComponent::ApplyGameplayEffectSpecToSelf(c
 	}
 
 	// Now that we have our own copy, apply our GEs that modify IncomingGEs
-	ActiveGameplayEffects.ApplyActiveEffectsTo(*OurCopyOfSpec, FModifierQualifier(BaseQualifier).Type(EGameplayMod::IncomingGE).IgnoreHandle(MyHandle));
+	if (!ActiveGameplayEffects.ApplyActiveEffectsTo(*OurCopyOfSpec, FModifierQualifier(BaseQualifier).Type(EGameplayMod::IncomingGE).IgnoreHandle(MyHandle)))
+	{
+		// We're immune to this effect
+		return FActiveGameplayEffectHandle();
+	}
 
 	// todo: apply some better logic to this so that we don't recalculate stacking effects as often
 	ActiveGameplayEffects.bNeedToRecalculateStacks = true;
diff --git a/Engine/Source/Runtime/SkillSystem/Private/GameplayEffect.cpp b/Engine/Source/Runtime/SkillSystem/Private/GameplayEffect.cpp
--- a/Engine/Source/Runtime/SkillSystem/Private/GameplayEffect.cpp
+++ b/Engine/Source/Runtime/SkillSystem/Private/GameplayEffect.cpp
@@ -21,7 +21,8 @@ FAggregator::FAllocationStats FAggregator::AllocationStats;
 UGameplayEffect::UGameplayEffect(const class FPostConstructInitializeProperties& PCIP)
 	: Super(PCIP)
 {
-
+	ChanceToApplyToTarget.SetValue(1.f);
+	ChanceToExecuteOnGameplayEffect.SetValue(1.f);
 }
 
 void UGameplayEffect::GetOwnedGameplayTags(FGameplayTagContainer& TagContainer) const
@@ -45,16 +46,20 @@ bool UGameplayEffect::AreApplicationTagRequirementsSatisfied(const FGameplayTagC
 //
 // --------------------------------------------------------------------------------------------------------------------------------------------------------
 
-FGameplayEffectSpec::FGameplayEffectSpec(const UGameplayEffect *InDef, AActor *Owner, float Level, const FGlobalCurveDataOverride *CurveData)
+FGameplayEffectSpec::FGameplayEffectSpec(const UGameplayEffect *InDef, AActor *Owner, AActor *Instigator, float Level, const FGlobalCurveDataOverride *CurveData)
 : Def(InDef)
 , ModifierLevel(TSharedPtr<FGameplayEffectLevelSpec>(new FGameplayEffectLevelSpec(Level, InDef->LevelInfo, Owner)))
 , Duration(new FAggregator(InDef->Duration.MakeFinalizedCopy(CurveData), ModifierLevel, SKILL_AGG_DEBUG(TEXT("%s Duration"), *InDef->GetName())))
 , Period(new FAggregator(InDef->Period.MakeFinalizedCopy(CurveData), ModifierLevel, SKILL_AGG_DEBUG(TEXT("%s Period"), *InDef->GetName())))
+, ChanceToApplyToTarget(new FAggregator(InDef->ChanceToApplyToTarget.MakeFinalizedCopy(CurveData), ModifierLevel, SKILL_AGG_DEBUG(TEXT("%s ChanceToApplyToTarget"), *InDef->GetName())))
+, ChanceToExecuteOnGameplayEffect(new FAggregator(InDef->ChanceToExecuteOnGameplayEffect.MakeFinalizedCopy(CurveData), ModifierLevel, SKILL_AGG_DEBUG(TEXT("%s ChanceToExecuteOnGameplayEffect"), *InDef->GetName())))
 , StackingPolicy(EGameplayEffectStackingPolicy::Unlimited)
 , StackedAttribName(NAME_None)
 {
 	Duration.Get()->RegisterLevelDependancies();
 	Period.Get()->RegisterLevelDependancies();
+	ChanceToApplyToTarget.Get()->RegisterLevelDependancies();
+	ChanceToExecuteOnGameplayEffect.Get()->RegisterLevelDependancies();
 
 	// check the stacking policy for this gameplay effect and warn if it's set up badly
 	if (Def->StackingPolicy != EGameplayEffectStackingPolicy::Unlimited)
@@ -140,9 +145,11 @@ FGameplayEffectSpec::FGameplayEffectSpec(const UGameplayEffect *InDef, AActor *O
 	{
 		for (UGameplayEffect* TargetDef : InDef->TargetEffects)
 		{
-			TargetEffectSpecs.Add(TSharedRef<FGameplayEffectSpec>(new FGameplayEffectSpec(TargetDef, Owner, Level, CurveData)));
+			TargetEffectSpecs.Add(TSharedRef<FGameplayEffectSpec>(new FGameplayEffectSpec(TargetDef, Owner, Instigator, Level, CurveData)));
 		}
 	}
+
+	InstigatorStack.AddInstigator(Instigator);
 }
 
 void FGameplayEffectSpec::InitModifiers(const FGlobalCurveDataOverride *CurveData, AActor *Owner, float Level)
@@ -181,11 +188,30 @@ int32 FGameplayEffectSpec::ApplyModifiersFrom(FGameplayEffectSpec &InSpec, const
 	int32 NumApplied = 0;
 	bool ShouldSnapshot = InSpec.ShouldApplyAsSnapshot(QualifierContext);
 
-	if (InSpec.Def && !InSpec.Def->AreGameplayEffectTagRequirementsSatisfied(Def))
+	if (!InSpec.Def || !InSpec.Def->AreGameplayEffectTagRequirementsSatisfied(Def))
 	{
+		// InSpec doesn't match this gameplay effect but if InSpec provides immunity we also need to check the modifiers because they can be canceled independent of the gameplay effect
+		if (InSpec.Def->AppliesImmunityTo == QualifierContext.Type())
+		{
+			for (int ii = 0; ii < Modifiers.Num(); ++ii)
+			{
+				const FGameplayModifierEvaluatedData& Data = Modifiers[ii].Aggregator.Get()->Evaluate();
+				if (InSpec.Def->AreGameplayEffectTagRequirementsSatisfied(Data.Tags))
+				{
+					Modifiers.RemoveAtSwap(ii);
+					--ii;
+				}
+			}
+		}
+
 		return 0;
 	}
 
+	if (InSpec.Def->AppliesImmunityTo == QualifierContext.Type())
+	{
+		return -1;
+	}
+
 	// Fixme: Use acceleration structures to speed up these types of lookups
 	// The called functions below are reliant on the InSpecs evaluated data. We should ideally only call evaluate once per mod.
 
@@ -222,9 +248,24 @@ int32 FGameplayEffectSpec::ApplyModifiersFrom(FGameplayEffectSpec &InSpec, const
 				break;
 			}
 
+			case EGameplayModEffect::ChanceApplyTarget:
+			{
+				ChanceToApplyToTarget.Get()->ApplyMod(InMod.Info.ModifierOp, InMod.Aggregator, ShouldSnapshot);
+				NumApplied++;
+				break;
+			}
+
+			case EGameplayModEffect::ChanceExecuteEffect:
+			{
+				ChanceToExecuteOnGameplayEffect.Get()->ApplyMod(InMod.Info.ModifierOp, InMod.Aggregator, ShouldSnapshot);
+				NumApplied++;
+				break;
+			}
+
 			case EGameplayModEffect::LinkedGameplayEffect:
 			{
 				TargetEffectSpecs.Add(InMod.TargetEffectSpec.ToSharedRef());
+				NumApplied++;
 				break;
 			}
 		}
@@ -263,6 +304,16 @@ float FGameplayEffectSpec::GetPeriod() const
 	return Period.Get()->Evaluate().Magnitude;
 }
 
+float FGameplayEffectSpec::GetChanceToApplyToTarget() const
+{
+	return ChanceToApplyToTarget.Get()->Evaluate().Magnitude;
+}
+
+float FGameplayEffectSpec::GetChanceToExecuteOnGameplayEffect() const
+{
+	return ChanceToExecuteOnGameplayEffect.Get()->Evaluate().Magnitude;
+}
+
 float FGameplayEffectSpec::GetMagnitude(const FGameplayAttribute &Attribute) const
 {
 	float CurrentMagnitude = 0.f;
@@ -357,7 +408,7 @@ FModifierSpec::FModifierSpec(const FGameplayModifierInfo &InInfo, TSharedPtr<FGa
 	Aggregator.Get()->RegisterLevelDependancies();
 	if (InInfo.TargetEffect)
 	{
-		TargetEffectSpec = TSharedPtr<FGameplayEffectSpec>(new FGameplayEffectSpec(InInfo.TargetEffect, Owner, Level, CurveData));
+		TargetEffectSpec = TSharedPtr<FGameplayEffectSpec>(new FGameplayEffectSpec(InInfo.TargetEffect, Owner, Owner, Level, CurveData));
 	}
 }
 
@@ -610,7 +661,7 @@ void FAggregator::ApplyMod(EGameplayModOp::Type ModType, FAggregatorRef Ref, boo
 }
 
 // Execute is intended to directly modify the base value of an aggregator.
-// However if the aggregator's base value is not static (scales with some outside influece)
+// However if the aggregator's base value is not static (scales with some outside influence)
 // we will treat this execute as an apply, but make our own copy of the passed in aggregator.
 // (so that it will essentially be permanent).
 void FAggregator::ExecuteModAggr(EGameplayModOp::Type ModType, FAggregatorRef Ref)
@@ -849,7 +900,7 @@ const FGameplayModifierEvaluatedData& FAggregator::Evaluate() const
 				// Our requirements on the mod
 				const FGameplayModifierEvaluatedData& ModEvaluatedData = Agg->Evaluate();
 
-				// We have already commited during this aggregation to not have certain tags
+				// We have already committed during this aggregation to not have certain tags
 				if (CommittedIgnoreTags.Num() > 0 && CommittedIgnoreTags.MatchesAny(ModEvaluatedData.Tags, false))
 				{
 					continue;
@@ -962,7 +1013,7 @@ void FAggregator::MakeUniqueDeep()
 	}
 }
 
-void FAggregator::CleaerAllDependancies()
+void FAggregator::ClearAllDependancies()
 {
 	Dependants.SetNum(0);
 	OnDirty.Unbind();
@@ -1048,7 +1099,7 @@ void FActiveGameplayEffect::PostReplicatedAdd(const struct FActiveGameplayEffect
 //
 // --------------------------------------------------------------------------------------------------------------------------------------------------------
 
-void FActiveGameplayEffectsContainer::ApplyActiveEffectsTo(OUT FGameplayEffectSpec &Spec, const FModifierQualifier &QualifierContext)
+bool FActiveGameplayEffectsContainer::ApplyActiveEffectsTo(OUT FGameplayEffectSpec &Spec, const FModifierQualifier &QualifierContext)
 {
 	FActiveGameplayEffectHandle().IsValid();
 
@@ -1060,8 +1111,20 @@ void FActiveGameplayEffectsContainer::ApplyActiveEffectsTo(OUT FGameplayEffectSp
 		{
 			continue;
 		}
-		Spec.ApplyModifiersFrom( ActiveEffect.Spec, QualifierContext );
+		// should we try to apply ActiveEffect to Spec?
+		float ChanceToExecute = ActiveEffect.Spec.GetChanceToExecuteOnGameplayEffect();
+		if ((ChanceToExecute < 1.f - SMALL_NUMBER) && (FMath::FRand() > ChanceToExecute))
+		{
+			continue;
+		}
+		if (Spec.ApplyModifiersFrom(ActiveEffect.Spec, QualifierContext) < 0)
+		{
+			// ActiveEffect provides immunity to Spec. No need to look at other active effects
+			return false;
+		}
 	}
+
+	return true;
 }
 
 /** This is the main function that applies/attaches a GameplayEffect on Attributes and ActiveGameplayEffects */
@@ -1084,7 +1147,7 @@ void FActiveGameplayEffectsContainer::ApplySpecToActiveEffectsAndAttributes(FGam
 
 		if (Mod.Info.ModifierType == EGameplayMod::ActiveGE)
 		{
-			SKILL_LOG_SCOPE(TEXT("Applying AttributeMod %s to ActiveEffects"),  *Mod.ToSimpleString(), *Mod.Info.Attribute.GetName());
+			SKILL_LOG_SCOPE(TEXT("Applying Mod %s to ActiveEffects"),  *Mod.ToSimpleString());
 
 			// TODO: Tag checks here
 
@@ -1118,6 +1181,41 @@ void FActiveGameplayEffectsContainer::ExecuteActiveEffectsFrom(const FGameplayEf
 		}
 	}
 
+	float ChanceToExecute = Spec.GetChanceToExecuteOnGameplayEffect();
+
+	// check if the new effect removes or modifies existing effects
+	if (Spec.Def->AppliesImmunityTo == EGameplayImmunity::ActiveGE)
+	{
+		for (int ii = 0; ii < GameplayEffects.Num(); ++ii)
+		{
+			// should we try to apply this to the current effect?
+			if ((ChanceToExecute < 1.f - SMALL_NUMBER) && (FMath::FRand() > ChanceToExecute))
+			{
+				continue;
+			}
+
+			if (Spec.Def->AreGameplayEffectTagRequirementsSatisfied(GameplayEffects[ii].Spec.Def))
+			{
+				GameplayEffects.RemoveAtSwap(ii);
+				--ii;
+				continue;
+			}
+
+			// we kept the effect, now check its mods to see if we should remove any of them
+			for (int jj = 0; jj < GameplayEffects[ii].Spec.Modifiers.Num(); ++jj)
+			{
+				const FGameplayModifierEvaluatedData& Data = GameplayEffects[ii].Spec.Modifiers[jj].Aggregator.Get()->Evaluate();
+				if (GameplayEffects[ii].Spec.Def->AreGameplayEffectTagRequirementsSatisfied(Data.Tags))
+				{
+					GameplayEffects[ii].Spec.Modifiers.RemoveAtSwap(jj);
+					--jj;
+				}
+			}
+		}
+	}
+
+	bool bModifiesActiveGEs = false;
+
 	for (const FModifierSpec &Mod : Spec.Modifiers)
 	{
 		if (Mod.Info.ModifierType == EGameplayMod::Attribute)
@@ -1177,23 +1275,31 @@ void FActiveGameplayEffectsContainer::ExecuteActiveEffectsFrom(const FGameplayEf
 		}
 		else if(Mod.Info.ModifierType == EGameplayMod::ActiveGE)
 		{
-			SKILL_LOG_SCOPE(TEXT("Executing Attribute Mod %s on ActiveEffects"), *Mod.ToSimpleString());
+			bModifiesActiveGEs = true;
+		}
+	}
 
-			// Todo: Tag checks here
+	// If any of the mods apply to active gameplay effects try to apply them now.
+	// We need to do this here because all of the modifiers either need to apply or not apply to each active gameplay effect based on their chance to execute on gameplay effects.
+	if (bModifiesActiveGEs)
+	{
+		for (FActiveGameplayEffect & ActiveEffect : GameplayEffects)
+		{
+			// Don't apply spec to itself
+			if (!QualifierContext.TestTarget(ActiveEffect.Handle))
+			{
+				continue;
+			}
 
-			// This modifies GEs that are currently active, so apply this to them...
-			for (FActiveGameplayEffect & ActiveEffect : GameplayEffects)
+			// should we try to apply this to the current effect?
+			if ((ChanceToExecute < 1.f - SMALL_NUMBER) && (FMath::FRand() > ChanceToExecute))
 			{
-				// Don't apply spec to itself
-				if (!QualifierContext.TestTarget(ActiveEffect.Handle))
-				{
-					continue;
-				}
+				continue;
+			}
 
-				if (ActiveEffect.Spec.ExecuteModifiersFrom( Spec, FModifierQualifier().Type(EGameplayMod::ActiveGE) ) > 0)
-				{
-					InvokeGameplayCueExecute = true;
-				}
+			if (ActiveEffect.Spec.ExecuteModifiersFrom(Spec, FModifierQualifier().Type(EGameplayMod::ActiveGE)) > 0)
+			{
+				InvokeGameplayCueExecute = true;
 			}
 		}
 	}
@@ -1379,7 +1485,7 @@ void FActiveGameplayEffectsContainer::PreDestroy()
 	{
 		if (It.Value.IsValid())
 		{
-			It.Value.Get()->CleaerAllDependancies();
+			It.Value.Get()->ClearAllDependancies();
 		}
 	}
 }
@@ -1587,9 +1693,7 @@ bool FActiveGameplayEffectsContainer::CanApplyAttributeModifiers(const UGameplay
 {
 	SCOPE_CYCLE_COUNTER(STAT_GameplayEffectsCanApplyAttributeModifiers);
 
-	FGameplayEffectSpec	Spec(GameplayEffect, Instigator, Level, Owner->GetCurveDataOverride());
-	Spec.Def = GameplayEffect;
-	Spec.InstigatorStack.AddInstigator(Instigator);
+	FGameplayEffectSpec	Spec(GameplayEffect, Instigator, Instigator, Level, Owner->GetCurveDataOverride());
 
 	ApplyActiveEffectsTo(Spec, FModifierQualifier().Type(EGameplayMod::IncomingGE));
 	
diff --git a/Engine/Source/Runtime/SkillSystem/Private/GameplayEffectTests.cpp b/Engine/Source/Runtime/SkillSystem/Private/GameplayEffectTests.cpp
--- a/Engine/Source/Runtime/SkillSystem/Private/GameplayEffectTests.cpp
+++ b/Engine/Source/Runtime/SkillSystem/Private/GameplayEffectTests.cpp
@@ -4315,6 +4315,1377 @@ bool GameplayEffectsTest_StackingAddingModifiers(UWorld *World, FAutomationTestB
 	return true;
 }
 
+
+bool GameplayEffectsTest_ImmunityIncoming(UWorld *World, FAutomationTestBase * Test)
+{
+	const float StartHealth = 100.f;
+	const float DamageValue = 5.f;
+
+	ASkillSystemTestPawn *SourceActor = World->SpawnActor<ASkillSystemTestPawn>();
+	ASkillSystemTestPawn *DestActor = World->SpawnActor<ASkillSystemTestPawn>();
+
+	UProperty *DamageProperty = FindFieldChecked<UProperty>(USkillSystemTestAttributeSet::StaticClass(), GET_MEMBER_NAME_CHECKED(USkillSystemTestAttributeSet, Damage));
+
+	UAttributeComponent * SourceComponent = SourceActor->AttributeComponent;
+	UAttributeComponent * DestComponent = DestActor->AttributeComponent;
+	SourceComponent->GetSet<USkillSystemTestAttributeSet>()->Health = StartHealth;
+	DestComponent->GetSet<USkillSystemTestAttributeSet>()->Health = StartHealth;
+
+	// Setup a GE to provide immunity from incoming GEs
+	{
+		SKILL_LOG_SCOPE(TEXT("Apply ImmunityBuff"))
+
+		UGameplayEffect * BaseImmunityEffect = Cast<UGameplayEffect>(StaticConstructObject(UGameplayEffect::StaticClass(), GetTransientPackage(), FName(TEXT("ImmunityBuff"))));
+		BaseImmunityEffect->Duration.SetValue(UGameplayEffect::INFINITE_DURATION);
+		BaseImmunityEffect->GameplayEffectRequiredTags.AddTag(IGameplayTagsModule::RequestGameplayTag(FName(TEXT("Immune"))));
+		BaseImmunityEffect->AppliesImmunityTo = EGameplayImmunity::IncomingGE;
+
+		// Apply to self
+		DestComponent->ApplyGameplayEffectToTarget(BaseImmunityEffect, DestComponent, 1.f);
+	}
+
+	// Apply Damage
+	{
+		SKILL_LOG_SCOPE(TEXT("Apply InstantDamage"))
+
+		UGameplayEffect * BaseDmgEffect = Cast<UGameplayEffect>(StaticConstructObject(UGameplayEffect::StaticClass(), GetTransientPackage(), FName(TEXT("Damage"))));
+		BaseDmgEffect->Modifiers.SetNum(1);
+		BaseDmgEffect->Modifiers[0].Magnitude.SetValue(DamageValue);
+		BaseDmgEffect->Modifiers[0].ModifierType = EGameplayMod::Attribute;
+		BaseDmgEffect->Modifiers[0].ModifierOp = EGameplayModOp::Additive;
+		BaseDmgEffect->Modifiers[0].Attribute.SetUProperty(DamageProperty);
+		BaseDmgEffect->Modifiers[0].OwnedTags.AddTag(IGameplayTagsModule::RequestGameplayTag(FName(TEXT("Damage.Basic"))));
+		BaseDmgEffect->Duration.Value = UGameplayEffect::INSTANT_APPLICATION;
+		BaseDmgEffect->GameplayEffectTags.AddTag(IGameplayTagsModule::RequestGameplayTag(FName(TEXT("Immune"))));
+
+		// Apply to target
+		SourceComponent->ApplyGameplayEffectToTarget(BaseDmgEffect, DestComponent, 1.f);
+
+		float ExpectedValue = StartHealth;
+		float ActualValue = DestComponent->GetSet<USkillSystemTestAttributeSet>()->Health;
+
+		Test->TestTrue(SKILL_TEST_TEXT("Buff Instant Damage Applied"), (ActualValue == ExpectedValue));
+		SKILL_LOG(Log, TEXT("Final Health: %.2f"), DestComponent->GetSet<USkillSystemTestAttributeSet>()->Health);
+	}
+
+	World->EditorDestroyActor(SourceActor, false);
+	World->EditorDestroyActor(DestActor, false);
+
+	return true;
+}
+
+bool GameplayEffectsTest_ImmunityOutgoing(UWorld *World, FAutomationTestBase * Test)
+{
+	const float StartHealth = 100.f;
+	const float DamageValue = 5.f;
+
+	ASkillSystemTestPawn *SourceActor = World->SpawnActor<ASkillSystemTestPawn>();
+	ASkillSystemTestPawn *DestActor = World->SpawnActor<ASkillSystemTestPawn>();
+
+	UProperty *DamageProperty = FindFieldChecked<UProperty>(USkillSystemTestAttributeSet::StaticClass(), GET_MEMBER_NAME_CHECKED(USkillSystemTestAttributeSet, Damage));
+
+	UAttributeComponent * SourceComponent = SourceActor->AttributeComponent;
+	UAttributeComponent * DestComponent = DestActor->AttributeComponent;
+	SourceComponent->GetSet<USkillSystemTestAttributeSet>()->Health = StartHealth;
+	DestComponent->GetSet<USkillSystemTestAttributeSet>()->Health = StartHealth;
+
+	// Setup a GE to provide immunity from outgoing GEs
+	{
+		SKILL_LOG_SCOPE(TEXT("Apply ImmunityBuff"))
+
+		UGameplayEffect * BaseImmunityEffect = Cast<UGameplayEffect>(StaticConstructObject(UGameplayEffect::StaticClass(), GetTransientPackage(), FName(TEXT("ImmunityBuff"))));
+		BaseImmunityEffect->Duration.SetValue(UGameplayEffect::INFINITE_DURATION);
+		BaseImmunityEffect->GameplayEffectRequiredTags.AddTag(IGameplayTagsModule::RequestGameplayTag(FName(TEXT("Immune"))));
+		BaseImmunityEffect->AppliesImmunityTo = EGameplayImmunity::OutgoingGE;
+
+		// Apply to self
+		SourceComponent->ApplyGameplayEffectToTarget(BaseImmunityEffect, SourceComponent, 1.f);
+	}
+
+	// Apply Damage
+	{
+		SKILL_LOG_SCOPE(TEXT("Apply InstantDamage"))
+
+		UGameplayEffect * BaseDmgEffect = Cast<UGameplayEffect>(StaticConstructObject(UGameplayEffect::StaticClass(), GetTransientPackage(), FName(TEXT("Damage"))));
+		BaseDmgEffect->Modifiers.SetNum(1);
+		BaseDmgEffect->Modifiers[0].Magnitude.SetValue(DamageValue);
+		BaseDmgEffect->Modifiers[0].ModifierType = EGameplayMod::Attribute;
+		BaseDmgEffect->Modifiers[0].ModifierOp = EGameplayModOp::Additive;
+		BaseDmgEffect->Modifiers[0].Attribute.SetUProperty(DamageProperty);
+		BaseDmgEffect->Modifiers[0].OwnedTags.AddTag(IGameplayTagsModule::RequestGameplayTag(FName(TEXT("Damage.Basic"))));
+		BaseDmgEffect->Duration.Value = UGameplayEffect::INSTANT_APPLICATION;
+		BaseDmgEffect->GameplayEffectTags.AddTag(IGameplayTagsModule::RequestGameplayTag(FName(TEXT("Immune"))));
+
+		// Apply to target
+		SourceComponent->ApplyGameplayEffectToTarget(BaseDmgEffect, DestComponent, 1.f);
+
+		float ExpectedValue = StartHealth;
+		float ActualValue = DestComponent->GetSet<USkillSystemTestAttributeSet>()->Health;
+
+		Test->TestTrue(SKILL_TEST_TEXT("Buff Instant Damage Applied"), (ActualValue == ExpectedValue));
+		SKILL_LOG(Log, TEXT("Final Health: %.2f"), DestComponent->GetSet<USkillSystemTestAttributeSet>()->Health);
+	}
+
+	World->EditorDestroyActor(SourceActor, false);
+	World->EditorDestroyActor(DestActor, false);
+
+	return true;
+}
+
+// A buff passes its tags to the base modifier. The base modifier can now be blocked by immunity
+// This would be bad to do in practice. We now have a gameplay effect that hangs around but doesn't do anything
+bool GameplayEffectsTest_ImmunityMod(UWorld *World, FAutomationTestBase * Test)
+{
+	const float StartHealth = 100.f;
+	const float DamageValue = 5.f;
+	const float BonusDamageMultiplier = 2.f;
+
+	ASkillSystemTestPawn *SourceActor = World->SpawnActor<ASkillSystemTestPawn>();
+	ASkillSystemTestPawn *DestActor = World->SpawnActor<ASkillSystemTestPawn>();
+
+	UProperty *DamageProperty = FindFieldChecked<UProperty>(USkillSystemTestAttributeSet::StaticClass(), GET_MEMBER_NAME_CHECKED(USkillSystemTestAttributeSet, Damage));
+
+	UAttributeComponent * SourceComponent = SourceActor->AttributeComponent;
+	UAttributeComponent * DestComponent = DestActor->AttributeComponent;
+	SourceComponent->GetSet<USkillSystemTestAttributeSet>()->Health = StartHealth;
+	DestComponent->GetSet<USkillSystemTestAttributeSet>()->Health = StartHealth;
+
+	// Setup a GE to provide immunity to buffed damage
+	{
+		SKILL_LOG_SCOPE(TEXT("Apply Immunity"))
+
+		UGameplayEffect * BaseImmunityEffect = Cast<UGameplayEffect>(StaticConstructObject(UGameplayEffect::StaticClass(), GetTransientPackage(), FName(TEXT("ImmunityBuff"))));
+		BaseImmunityEffect->Duration.SetValue(UGameplayEffect::INFINITE_DURATION);
+		BaseImmunityEffect->GameplayEffectRequiredTags.AddTag(IGameplayTagsModule::RequestGameplayTag(FName(TEXT("Damage.Buff"))));
+		BaseImmunityEffect->AppliesImmunityTo = EGameplayImmunity::IncomingGE;
+
+		// Apply to self
+		DestComponent->ApplyGameplayEffectToTarget(BaseImmunityEffect, DestComponent, 1.f);
+	}
+
+	// Setup a GE to modify OutgoingGEs
+	{
+		SKILL_LOG_SCOPE(TEXT("Apply DamageBuff"))
+
+		UGameplayEffect * BaseDmgEffect = Cast<UGameplayEffect>(StaticConstructObject(UGameplayEffect::StaticClass(), GetTransientPackage(), FName(TEXT("DamageBuff"))));
+		BaseDmgEffect->Modifiers.SetNum(1);
+		BaseDmgEffect->Modifiers[0].Magnitude.SetValue(BonusDamageMultiplier);
+		BaseDmgEffect->Modifiers[0].ModifierType = EGameplayMod::OutgoingGE;
+		BaseDmgEffect->Modifiers[0].ModifierOp = EGameplayModOp::Multiplicitive;
+		BaseDmgEffect->Modifiers[0].Attribute.SetUProperty(DamageProperty);
+		BaseDmgEffect->Modifiers[0].OwnedTags.AddTag(IGameplayTagsModule::RequestGameplayTag(FName(TEXT("Damage.Buff"))));
+		BaseDmgEffect->Duration.SetValue(UGameplayEffect::INFINITE_DURATION);
+
+		// Apply to self
+		SourceComponent->ApplyGameplayEffectToTarget(BaseDmgEffect, SourceComponent, 1.f);
+	}
+
+	// Apply Damage
+	{
+		SKILL_LOG_SCOPE(TEXT("Apply InstantDamage"))
+
+		UGameplayEffect * BaseDmgEffect = Cast<UGameplayEffect>(StaticConstructObject(UGameplayEffect::StaticClass(), GetTransientPackage(), FName(TEXT("Damage"))));
+		BaseDmgEffect->Modifiers.SetNum(1);
+		BaseDmgEffect->Modifiers[0].Magnitude.SetValue(DamageValue);
+		BaseDmgEffect->Modifiers[0].ModifierType = EGameplayMod::Attribute;
+		BaseDmgEffect->Modifiers[0].ModifierOp = EGameplayModOp::Additive;
+		BaseDmgEffect->Modifiers[0].Attribute.SetUProperty(DamageProperty);
+		BaseDmgEffect->Modifiers[0].OwnedTags.AddTag(IGameplayTagsModule::RequestGameplayTag(FName(TEXT("Damage.Basic"))));
+		BaseDmgEffect->Duration.Value = UGameplayEffect::INSTANT_APPLICATION;
+
+		// Apply to target
+		SourceComponent->ApplyGameplayEffectToTarget(BaseDmgEffect, DestComponent, 1.f);
+
+		float ExpectedValue = StartHealth;
+		float ActualValue = DestComponent->GetSet<USkillSystemTestAttributeSet>()->Health;
+
+		Test->TestTrue(SKILL_TEST_TEXT("Buff Instant Damage Applied"), (ActualValue == ExpectedValue));
+		SKILL_LOG(Log, TEXT("Final Health: %.2f"), DestComponent->GetSet<USkillSystemTestAttributeSet>()->Health);
+	}
+
+	World->EditorDestroyActor(SourceActor, false);
+	World->EditorDestroyActor(DestActor, false);
+
+	return true;
+}
+
+// A buff is added in the form of a new gameplay effect being attached to the base gameplay effect.
+// The base effect applies but immunity stops the buff.
+bool GameplayEffectsTest_ImmunityBlockedBuff(UWorld *World, FAutomationTestBase * Test)
+{
+	const float StartHealth = 100.f;
+	const float DamageValue = 5.f;
+
+	ASkillSystemTestPawn *SourceActor = World->SpawnActor<ASkillSystemTestPawn>();
+	ASkillSystemTestPawn *DestActor = World->SpawnActor<ASkillSystemTestPawn>();
+
+	UProperty *DamageProperty = FindFieldChecked<UProperty>(USkillSystemTestAttributeSet::StaticClass(), GET_MEMBER_NAME_CHECKED(USkillSystemTestAttributeSet, Damage));
+
+	UAttributeComponent * SourceComponent = SourceActor->AttributeComponent;
+	UAttributeComponent * DestComponent = DestActor->AttributeComponent;
+	SourceComponent->GetSet<USkillSystemTestAttributeSet>()->Health = StartHealth;
+	DestComponent->GetSet<USkillSystemTestAttributeSet>()->Health = StartHealth;
+
+	// Setup a GE to provide immunity to buffed damage
+	{
+		SKILL_LOG_SCOPE(TEXT("Apply Immunity"))
+
+		UGameplayEffect * BaseImmunityEffect = Cast<UGameplayEffect>(StaticConstructObject(UGameplayEffect::StaticClass(), GetTransientPackage(), FName(TEXT("ImmunityBuff"))));
+		BaseImmunityEffect->Duration.SetValue(UGameplayEffect::INFINITE_DURATION);
+		BaseImmunityEffect->GameplayEffectRequiredTags.AddTag(IGameplayTagsModule::RequestGameplayTag(FName(TEXT("Damage.Buff"))));
+		BaseImmunityEffect->AppliesImmunityTo = EGameplayImmunity::IncomingGE;
+
+		// Apply to self
+		DestComponent->ApplyGameplayEffectToTarget(BaseImmunityEffect, DestComponent, 1.f);
+	}
+
+	// Apply base damage and a buff that causes extra damage
+	{
+		SKILL_LOG_SCOPE(TEXT("Apply DamageBuff and InstantDamage"))
+
+		UGameplayEffect * BuffEffect = Cast<UGameplayEffect>(StaticConstructObject(UGameplayEffect::StaticClass(), GetTransientPackage(), FName(TEXT("DamageBuff"))));
+		BuffEffect->Modifiers.SetNum(1);
+		BuffEffect->Modifiers[0].Magnitude.SetValue(DamageValue);
+		BuffEffect->Modifiers[0].ModifierType = EGameplayMod::Attribute;
+		BuffEffect->Modifiers[0].ModifierOp = EGameplayModOp::Additive;
+		BuffEffect->Modifiers[0].Attribute.SetUProperty(DamageProperty);
+		BuffEffect->GameplayEffectTags.AddTag(IGameplayTagsModule::RequestGameplayTag(FName(TEXT("Damage.Buff"))));
+		BuffEffect->Duration.SetValue(UGameplayEffect::INSTANT_APPLICATION);
+
+		UGameplayEffect * BaseDmgEffect = Cast<UGameplayEffect>(StaticConstructObject(UGameplayEffect::StaticClass(), GetTransientPackage(), FName(TEXT("Damage"))));
+		BaseDmgEffect->Modifiers.SetNum(1);
+		BaseDmgEffect->Modifiers[0].Magnitude.SetValue(DamageValue);
+		BaseDmgEffect->Modifiers[0].ModifierType = EGameplayMod::Attribute;
+		BaseDmgEffect->Modifiers[0].ModifierOp = EGameplayModOp::Additive;
+		BaseDmgEffect->Modifiers[0].Attribute.SetUProperty(DamageProperty);
+		BaseDmgEffect->GameplayEffectTags.AddTag(IGameplayTagsModule::RequestGameplayTag(FName(TEXT("Damage.Basic"))));
+		BaseDmgEffect->Duration.Value = UGameplayEffect::INSTANT_APPLICATION;
+		BaseDmgEffect->TargetEffects.Add(BuffEffect);
+
+		// Apply to target
+		SourceComponent->ApplyGameplayEffectToTarget(BaseDmgEffect, DestComponent, 1.f);
+
+		float ExpectedValue = (StartHealth - DamageValue);
+		float ActualValue = DestComponent->GetSet<USkillSystemTestAttributeSet>()->Health;
+
+		Test->TestTrue(SKILL_TEST_TEXT("Buff Instant Damage Applied. ActualValue: %.2f. ExpectedValue: %.2f.", ActualValue, ExpectedValue), (ActualValue == ExpectedValue));
+		SKILL_LOG(Log, TEXT("Final Health: %.2f"), DestComponent->GetSet<USkillSystemTestAttributeSet>()->Health);
+	}
+
+	World->EditorDestroyActor(SourceActor, false);
+	World->EditorDestroyActor(DestActor, false);
+
+	return true;
+}
+
+// A buff is added in the form of a new gameplay effect being attached to the base gameplay effect.
+// The base effect is stopped by immunity, the buff should also be stopped despite not matching the immunity tags
+bool GameplayEffectsTest_ImmunityBlockedBaseAndBuff(UWorld *World, FAutomationTestBase * Test)
+{
+	const float StartHealth = 100.f;
+	const float DamageValue = 5.f;
+
+	ASkillSystemTestPawn *SourceActor = World->SpawnActor<ASkillSystemTestPawn>();
+	ASkillSystemTestPawn *DestActor = World->SpawnActor<ASkillSystemTestPawn>();
+
+	UProperty *DamageProperty = FindFieldChecked<UProperty>(USkillSystemTestAttributeSet::StaticClass(), GET_MEMBER_NAME_CHECKED(USkillSystemTestAttributeSet, Damage));
+
+	UAttributeComponent * SourceComponent = SourceActor->AttributeComponent;
+	UAttributeComponent * DestComponent = DestActor->AttributeComponent;
+	SourceComponent->GetSet<USkillSystemTestAttributeSet>()->Health = StartHealth;
+	DestComponent->GetSet<USkillSystemTestAttributeSet>()->Health = StartHealth;
+
+	// Setup a GE to provide immunity to buffed damage
+	{
+		SKILL_LOG_SCOPE(TEXT("Apply Immunity"))
+
+		UGameplayEffect * BaseImmunityEffect = Cast<UGameplayEffect>(StaticConstructObject(UGameplayEffect::StaticClass(), GetTransientPackage(), FName(TEXT("ImmunityBuff"))));
+		BaseImmunityEffect->Duration.SetValue(UGameplayEffect::INFINITE_DURATION);
+		BaseImmunityEffect->GameplayEffectRequiredTags.AddTag(IGameplayTagsModule::RequestGameplayTag(FName(TEXT("Damage.Basic"))));
+		BaseImmunityEffect->AppliesImmunityTo = EGameplayImmunity::IncomingGE;
+
+		// Apply to self
+		DestComponent->ApplyGameplayEffectToTarget(BaseImmunityEffect, DestComponent, 1.f);
+	}
+
+	// Apply base damage and a buff that causes extra damage
+	{
+		SKILL_LOG_SCOPE(TEXT("Apply DamageBuff and InstantDamage"))
+
+		UGameplayEffect * BuffEffect = Cast<UGameplayEffect>(StaticConstructObject(UGameplayEffect::StaticClass(), GetTransientPackage(), FName(TEXT("DamageBuff"))));
+		BuffEffect->Modifiers.SetNum(1);
+		BuffEffect->Modifiers[0].Magnitude.SetValue(DamageValue);
+		BuffEffect->Modifiers[0].ModifierType = EGameplayMod::Attribute;
+		BuffEffect->Modifiers[0].ModifierOp = EGameplayModOp::Additive;
+		BuffEffect->Modifiers[0].Attribute.SetUProperty(DamageProperty);
+		BuffEffect->GameplayEffectTags.AddTag(IGameplayTagsModule::RequestGameplayTag(FName(TEXT("Damage.Buff"))));
+		BuffEffect->Duration.SetValue(UGameplayEffect::INSTANT_APPLICATION);
+
+		UGameplayEffect * BaseDmgEffect = Cast<UGameplayEffect>(StaticConstructObject(UGameplayEffect::StaticClass(), GetTransientPackage(), FName(TEXT("Damage"))));
+		BaseDmgEffect->Modifiers.SetNum(1);
+		BaseDmgEffect->Modifiers[0].Magnitude.SetValue(DamageValue);
+		BaseDmgEffect->Modifiers[0].ModifierType = EGameplayMod::Attribute;
+		BaseDmgEffect->Modifiers[0].ModifierOp = EGameplayModOp::Additive;
+		BaseDmgEffect->Modifiers[0].Attribute.SetUProperty(DamageProperty);
+		BaseDmgEffect->GameplayEffectTags.AddTag(IGameplayTagsModule::RequestGameplayTag(FName(TEXT("Damage.Basic"))));
+		BaseDmgEffect->Duration.Value = UGameplayEffect::INSTANT_APPLICATION;
+		BaseDmgEffect->TargetEffects.Add(BuffEffect);
+
+		// Apply to target
+		SourceComponent->ApplyGameplayEffectToTarget(BaseDmgEffect, DestComponent, 1.f);
+
+		float ExpectedValue = StartHealth;
+		float ActualValue = DestComponent->GetSet<USkillSystemTestAttributeSet>()->Health;
+
+		Test->TestTrue(SKILL_TEST_TEXT("Buff Instant Damage Applied"), (ActualValue == ExpectedValue));
+		SKILL_LOG(Log, TEXT("Final Health: %.2f"), DestComponent->GetSet<USkillSystemTestAttributeSet>()->Health);
+	}
+
+	World->EditorDestroyActor(SourceActor, false);
+	World->EditorDestroyActor(DestActor, false);
+
+	return true;
+}
+
+// Removes an active gameplay effect from the target
+bool GameplayEffectsTest_ImmunityActiveGE(UWorld *World, FAutomationTestBase * Test)
+{
+	const float StartHealth = 100.f;
+	const float DamageValue = 5.f;
+	const float DamageProtectionDivisor = 2.f;
+
+	ASkillSystemTestPawn *SourceActor = World->SpawnActor<ASkillSystemTestPawn>();
+	ASkillSystemTestPawn *DestActor = World->SpawnActor<ASkillSystemTestPawn>();
+
+	UProperty *DamageProperty = FindFieldChecked<UProperty>(USkillSystemTestAttributeSet::StaticClass(), GET_MEMBER_NAME_CHECKED(USkillSystemTestAttributeSet, Damage));
+
+	UAttributeComponent * SourceComponent = SourceActor->AttributeComponent;
+	UAttributeComponent * DestComponent = DestActor->AttributeComponent;
+	SourceComponent->GetSet<USkillSystemTestAttributeSet>()->Health = StartHealth;
+	DestComponent->GetSet<USkillSystemTestAttributeSet>()->Health = StartHealth;
+
+	// Apply a buff to reduce incoming damage on DestComponent
+	{
+		UGameplayEffect * BuffEffect = Cast<UGameplayEffect>(StaticConstructObject(UGameplayEffect::StaticClass(), GetTransientPackage(), FName(TEXT("DamageBuff"))));
+		BuffEffect->Modifiers.SetNum(1);
+		BuffEffect->Modifiers[0].Magnitude.SetValue(DamageProtectionDivisor);
+		BuffEffect->Modifiers[0].ModifierType = EGameplayMod::IncomingGE;
+		BuffEffect->Modifiers[0].ModifierOp = EGameplayModOp::Division;
+		BuffEffect->Modifiers[0].Attribute.SetUProperty(DamageProperty);
+		BuffEffect->GameplayEffectTags.AddTag(IGameplayTagsModule::RequestGameplayTag(FName(TEXT("Damage.Reduce"))));
+		BuffEffect->Duration.SetValue(UGameplayEffect::INFINITE_DURATION);
+
+		// Apply to self
+		DestComponent->ApplyGameplayEffectToTarget(BuffEffect, DestComponent, 1.f);
+	}
+
+	// Apply Damage to verify the buff is working
+	{
+		SKILL_LOG_SCOPE(TEXT("Apply InstantDamage"))
+
+		UGameplayEffect * BaseDmgEffect = Cast<UGameplayEffect>(StaticConstructObject(UGameplayEffect::StaticClass(), GetTransientPackage(), FName(TEXT("Damage"))));
+		BaseDmgEffect->Modifiers.SetNum(1);
+		BaseDmgEffect->Modifiers[0].Magnitude.SetValue(DamageValue);
+		BaseDmgEffect->Modifiers[0].ModifierType = EGameplayMod::Attribute;
+		BaseDmgEffect->Modifiers[0].ModifierOp = EGameplayModOp::Additive;
+		BaseDmgEffect->Modifiers[0].Attribute.SetUProperty(DamageProperty);
+		BaseDmgEffect->Modifiers[0].OwnedTags.AddTag(IGameplayTagsModule::RequestGameplayTag(FName(TEXT("Damage.Basic"))));
+		BaseDmgEffect->Duration.Value = UGameplayEffect::INSTANT_APPLICATION;
+
+		// Apply to target
+		SourceComponent->ApplyGameplayEffectToTarget(BaseDmgEffect, DestComponent, 1.f);
+
+		float ExpectedValue = StartHealth - (DamageValue / DamageProtectionDivisor);
+		float ActualValue = DestComponent->GetSet<USkillSystemTestAttributeSet>()->Health;
+
+		Test->TestTrue(SKILL_TEST_TEXT("Buff Instant Damage Applied"), (ActualValue == ExpectedValue));
+		SKILL_LOG(Log, TEXT("Final Health: %.2f"), DestComponent->GetSet<USkillSystemTestAttributeSet>()->Health);
+	}
+
+	// Setup a GE to remove the buff
+	{
+		SKILL_LOG_SCOPE(TEXT("Apply Immunity"))
+
+		UGameplayEffect * BaseImmunityEffect = Cast<UGameplayEffect>(StaticConstructObject(UGameplayEffect::StaticClass(), GetTransientPackage(), FName(TEXT("ImmunityBuff"))));
+		BaseImmunityEffect->Duration.SetValue(UGameplayEffect::INSTANT_APPLICATION);
+		BaseImmunityEffect->GameplayEffectRequiredTags.AddTag(IGameplayTagsModule::RequestGameplayTag(FName(TEXT("Damage.Reduce"))));
+		BaseImmunityEffect->AppliesImmunityTo = EGameplayImmunity::ActiveGE;
+
+		// Apply
+		SourceComponent->ApplyGameplayEffectToTarget(BaseImmunityEffect, DestComponent, 1.f);
+	}
+
+	// Apply Damage
+	{
+		SKILL_LOG_SCOPE(TEXT("Apply InstantDamage"))
+
+		UGameplayEffect * BaseDmgEffect = Cast<UGameplayEffect>(StaticConstructObject(UGameplayEffect::StaticClass(), GetTransientPackage(), FName(TEXT("Damage"))));
+		BaseDmgEffect->Modifiers.SetNum(1);
+		BaseDmgEffect->Modifiers[0].Magnitude.SetValue(DamageValue);
+		BaseDmgEffect->Modifiers[0].ModifierType = EGameplayMod::Attribute;
+		BaseDmgEffect->Modifiers[0].ModifierOp = EGameplayModOp::Additive;
+		BaseDmgEffect->Modifiers[0].Attribute.SetUProperty(DamageProperty);
+		BaseDmgEffect->Modifiers[0].OwnedTags.AddTag(IGameplayTagsModule::RequestGameplayTag(FName(TEXT("Damage.Basic"))));
+		BaseDmgEffect->Duration.Value = UGameplayEffect::INSTANT_APPLICATION;
+
+		// Apply to target
+		SourceComponent->ApplyGameplayEffectToTarget(BaseDmgEffect, DestComponent, 1.f);
+
+		float ExpectedValue = StartHealth - (DamageValue / DamageProtectionDivisor) - DamageValue;
+		float ActualValue = DestComponent->GetSet<USkillSystemTestAttributeSet>()->Health;
+
+		Test->TestTrue(SKILL_TEST_TEXT("Buff Instant Damage Applied"), (ActualValue == ExpectedValue));
+		SKILL_LOG(Log, TEXT("Final Health: %.2f"), DestComponent->GetSet<USkillSystemTestAttributeSet>()->Health);
+	}
+
+	World->EditorDestroyActor(SourceActor, false);
+	World->EditorDestroyActor(DestActor, false);
+
+	return true;
+}
+
+// removes a modifier on an active gameplay effect from the target
+bool GameplayEffectsTest_ImmunityActiveMod(UWorld *World, FAutomationTestBase * Test)
+{
+	const float StartHealth = 100.f;
+	const float DamageValue = 5.f;
+	const float DamageProtectionDivisor = 2.f;
+
+	ASkillSystemTestPawn *SourceActor = World->SpawnActor<ASkillSystemTestPawn>();
+	ASkillSystemTestPawn *DestActor = World->SpawnActor<ASkillSystemTestPawn>();
+
+	UProperty *DamageProperty = FindFieldChecked<UProperty>(USkillSystemTestAttributeSet::StaticClass(), GET_MEMBER_NAME_CHECKED(USkillSystemTestAttributeSet, Damage));
+
+	UAttributeComponent * SourceComponent = SourceActor->AttributeComponent;
+	UAttributeComponent * DestComponent = DestActor->AttributeComponent;
+	SourceComponent->GetSet<USkillSystemTestAttributeSet>()->Health = StartHealth;
+	DestComponent->GetSet<USkillSystemTestAttributeSet>()->Health = StartHealth;
+
+	// Apply a buff to reduce incoming damage on DestComponent
+	{
+		UGameplayEffect * BuffEffect = Cast<UGameplayEffect>(StaticConstructObject(UGameplayEffect::StaticClass(), GetTransientPackage(), FName(TEXT("DamageBuff"))));
+		BuffEffect->Modifiers.SetNum(1);
+		BuffEffect->Modifiers[0].Magnitude.SetValue(DamageProtectionDivisor);
+		BuffEffect->Modifiers[0].ModifierType = EGameplayMod::IncomingGE;
+		BuffEffect->Modifiers[0].ModifierOp = EGameplayModOp::Division;
+		BuffEffect->Modifiers[0].Attribute.SetUProperty(DamageProperty);
+		BuffEffect->Modifiers[0].OwnedTags.AddTag(IGameplayTagsModule::RequestGameplayTag(FName(TEXT("Damage.Reduce"))));
+		BuffEffect->Duration.SetValue(UGameplayEffect::INFINITE_DURATION);
+
+		// Apply to self
+		DestComponent->ApplyGameplayEffectToTarget(BuffEffect, DestComponent, 1.f);
+	}
+
+	// Apply Damage to verify the buff is working
+	{
+		SKILL_LOG_SCOPE(TEXT("Apply InstantDamage"))
+
+		UGameplayEffect * BaseDmgEffect = Cast<UGameplayEffect>(StaticConstructObject(UGameplayEffect::StaticClass(), GetTransientPackage(), FName(TEXT("Damage"))));
+		BaseDmgEffect->Modifiers.SetNum(1);
+		BaseDmgEffect->Modifiers[0].Magnitude.SetValue(DamageValue);
+		BaseDmgEffect->Modifiers[0].ModifierType = EGameplayMod::Attribute;
+		BaseDmgEffect->Modifiers[0].ModifierOp = EGameplayModOp::Additive;
+		BaseDmgEffect->Modifiers[0].Attribute.SetUProperty(DamageProperty);
+		BaseDmgEffect->Modifiers[0].OwnedTags.AddTag(IGameplayTagsModule::RequestGameplayTag(FName(TEXT("Damage.Basic"))));
+		BaseDmgEffect->Duration.Value = UGameplayEffect::INSTANT_APPLICATION;
+
+		// Apply to target
+		SourceComponent->ApplyGameplayEffectToTarget(BaseDmgEffect, DestComponent, 1.f);
+
+		float ExpectedValue = StartHealth - (DamageValue / DamageProtectionDivisor);
+		float ActualValue = DestComponent->GetSet<USkillSystemTestAttributeSet>()->Health;
+
+		Test->TestTrue(SKILL_TEST_TEXT("Buff Instant Damage Applied"), (ActualValue == ExpectedValue));
+		SKILL_LOG(Log, TEXT("Final Health: %.2f"), DestComponent->GetSet<USkillSystemTestAttributeSet>()->Health);
+	}
+
+	// Setup a GE to remove the buff
+	{
+		SKILL_LOG_SCOPE(TEXT("Apply Immunity"))
+
+		UGameplayEffect * BaseImmunityEffect = Cast<UGameplayEffect>(StaticConstructObject(UGameplayEffect::StaticClass(), GetTransientPackage(), FName(TEXT("ImmunityBuff"))));
+		BaseImmunityEffect->Duration.SetValue(UGameplayEffect::INSTANT_APPLICATION);
+		BaseImmunityEffect->GameplayEffectRequiredTags.AddTag(IGameplayTagsModule::RequestGameplayTag(FName(TEXT("Damage.Reduce"))));
+		BaseImmunityEffect->AppliesImmunityTo = EGameplayImmunity::ActiveGE;
+
+		// Apply
+		SourceComponent->ApplyGameplayEffectToTarget(BaseImmunityEffect, DestComponent, 1.f);
+	}
+
+	// Apply Damage
+	{
+		SKILL_LOG_SCOPE(TEXT("Apply InstantDamage"))
+
+		UGameplayEffect * BaseDmgEffect = Cast<UGameplayEffect>(StaticConstructObject(UGameplayEffect::StaticClass(), GetTransientPackage(), FName(TEXT("Damage"))));
+		BaseDmgEffect->Modifiers.SetNum(1);
+		BaseDmgEffect->Modifiers[0].Magnitude.SetValue(DamageValue);
+		BaseDmgEffect->Modifiers[0].ModifierType = EGameplayMod::Attribute;
+		BaseDmgEffect->Modifiers[0].ModifierOp = EGameplayModOp::Additive;
+		BaseDmgEffect->Modifiers[0].Attribute.SetUProperty(DamageProperty);
+		BaseDmgEffect->Modifiers[0].OwnedTags.AddTag(IGameplayTagsModule::RequestGameplayTag(FName(TEXT("Damage.Basic"))));
+		BaseDmgEffect->Duration.Value = UGameplayEffect::INSTANT_APPLICATION;
+
+		// Apply to target
+		SourceComponent->ApplyGameplayEffectToTarget(BaseDmgEffect, DestComponent, 1.f);
+
+		float ExpectedValue = StartHealth - (DamageValue / DamageProtectionDivisor) - DamageValue;
+		float ActualValue = DestComponent->GetSet<USkillSystemTestAttributeSet>()->Health;
+
+		Test->TestTrue(SKILL_TEST_TEXT("Buff Instant Damage Applied"), (ActualValue == ExpectedValue));
+		SKILL_LOG(Log, TEXT("Final Health: %.2f"), DestComponent->GetSet<USkillSystemTestAttributeSet>()->Health);
+	}
+
+	World->EditorDestroyActor(SourceActor, false);
+	World->EditorDestroyActor(DestActor, false);
+
+	return true;
+}
+
+
+bool GameplayEffectsTest_ChanceToApplyToTarget(UWorld *World, FAutomationTestBase * Test)
+{
+	const float StartHealth = 100.f;
+	const float DamageValue = 5.f;
+
+	ASkillSystemTestPawn *SourceActor = World->SpawnActor<ASkillSystemTestPawn>();
+	ASkillSystemTestPawn *DestActor = World->SpawnActor<ASkillSystemTestPawn>();
+
+	UProperty *DamageProperty = FindFieldChecked<UProperty>(USkillSystemTestAttributeSet::StaticClass(), GET_MEMBER_NAME_CHECKED(USkillSystemTestAttributeSet, Damage));
+
+	UAttributeComponent * SourceComponent = Cast<UAttributeComponent>(SourceActor->CreateComponentFromTemplate(UAttributeComponent::StaticClass()->GetDefaultObject<UAttributeComponent>(), FString(TEXT("AttributeComponent"))));
+	UAttributeComponent * DestComponent = Cast<UAttributeComponent>(DestActor->CreateComponentFromTemplate(UAttributeComponent::StaticClass()->GetDefaultObject<UAttributeComponent>(), FString(TEXT("AttributeComponent"))));
+	SourceComponent->GetSet<USkillSystemTestAttributeSet>()->Health = StartHealth;
+	DestComponent->GetSet<USkillSystemTestAttributeSet>()->Health = StartHealth;
+
+	{
+		SKILL_LOG_SCOPE(TEXT("Apply Damage mod"))
+
+		UGameplayEffect * BaseDmgEffect = Cast<UGameplayEffect>(StaticConstructObject(UGameplayEffect::StaticClass(), GetTransientPackage(), FName(TEXT("Damage"))));
+		BaseDmgEffect->Modifiers.SetNum(1);
+		BaseDmgEffect->Modifiers[0].Magnitude.SetValue(DamageValue);
+		BaseDmgEffect->Modifiers[0].ModifierType = EGameplayMod::Attribute;
+		BaseDmgEffect->Modifiers[0].ModifierOp = EGameplayModOp::Additive;
+		BaseDmgEffect->Modifiers[0].Attribute.SetUProperty(DamageProperty);
+		BaseDmgEffect->Modifiers[0].OwnedTags.AddTag(IGameplayTagsModule::RequestGameplayTag(FName(TEXT("Damage.Basic"))));
+		BaseDmgEffect->Duration.SetValue(UGameplayEffect::INSTANT_APPLICATION);
+		BaseDmgEffect->ChanceToApplyToTarget.SetValue(1.f);
+
+		// Apply to target
+		SourceComponent->ApplyGameplayEffectToTarget(BaseDmgEffect, DestComponent, 1.f);
+
+		float ExpectedValue = (StartHealth - DamageValue);
+		float ActualValue = DestComponent->GetSet<USkillSystemTestAttributeSet>()->Health;
+
+		Test->TestTrue(SKILL_TEST_TEXT("Damage Applied. %.2f == %.2f", ActualValue, ExpectedValue), ActualValue == ExpectedValue);
+	}
+
+	{
+		SKILL_LOG_SCOPE(TEXT("Apply Damage mod"))
+
+		UGameplayEffect * BaseDmgEffect = Cast<UGameplayEffect>(StaticConstructObject(UGameplayEffect::StaticClass(), GetTransientPackage(), FName(TEXT("Damage"))));
+		BaseDmgEffect->Modifiers.SetNum(1);
+		BaseDmgEffect->Modifiers[0].Magnitude.SetValue(DamageValue);
+		BaseDmgEffect->Modifiers[0].ModifierType = EGameplayMod::Attribute;
+		BaseDmgEffect->Modifiers[0].ModifierOp = EGameplayModOp::Additive;
+		BaseDmgEffect->Modifiers[0].Attribute.SetUProperty(DamageProperty);
+		BaseDmgEffect->Modifiers[0].OwnedTags.AddTag(IGameplayTagsModule::RequestGameplayTag(FName(TEXT("Damage.Basic"))));
+		BaseDmgEffect->Duration.SetValue(UGameplayEffect::INSTANT_APPLICATION);
+		BaseDmgEffect->ChanceToApplyToTarget.SetValue(0.f);
+
+		// Apply to target
+		FActiveGameplayEffectHandle Handle = SourceComponent->ApplyGameplayEffectToTarget(BaseDmgEffect, DestComponent, 1.f);
+
+		float ExpectedValue = (StartHealth - DamageValue);
+		float ActualValue = DestComponent->GetSet<USkillSystemTestAttributeSet>()->Health;
+
+		Test->TestTrue(SKILL_TEST_TEXT("Damage Applied. %.2f == %.2f", ActualValue, ExpectedValue), ActualValue == ExpectedValue);
+		Test->TestFalse(SKILL_TEST_TEXT("Effect applied to target when chance was 0.f"), Handle.IsValid());
+	}
+
+	World->EditorDestroyActor(SourceActor, false);
+	World->EditorDestroyActor(DestActor, false);
+
+	return true;
+}
+
+bool GameplayEffectsTest_ChanceToExecuteOnActiveGEMod(UWorld *World, FAutomationTestBase * Test)
+{
+	const float StartHealth = 100.f;
+	const float DamageValue = 5.f;
+	const float BonusDamageMultiplier = 2.f;
+	const float ExtraDamageMultiplier = 1.f;
+
+	ASkillSystemTestPawn *SourceActor = World->SpawnActor<ASkillSystemTestPawn>();
+	ASkillSystemTestPawn *DestActor = World->SpawnActor<ASkillSystemTestPawn>();
+
+	UProperty *DamageProperty = FindFieldChecked<UProperty>(USkillSystemTestAttributeSet::StaticClass(), GET_MEMBER_NAME_CHECKED(USkillSystemTestAttributeSet, Damage));
+
+	UAttributeComponent * SourceComponent = Cast<UAttributeComponent>(SourceActor->CreateComponentFromTemplate(UAttributeComponent::StaticClass()->GetDefaultObject<UAttributeComponent>(), FString(TEXT("AttributeComponent"))));
+	UAttributeComponent * DestComponent = Cast<UAttributeComponent>(DestActor->CreateComponentFromTemplate(UAttributeComponent::StaticClass()->GetDefaultObject<UAttributeComponent>(), FString(TEXT("AttributeComponent"))));
+	SourceComponent->GetSet<USkillSystemTestAttributeSet>()->Health = StartHealth;
+	DestComponent->GetSet<USkillSystemTestAttributeSet>()->Health = StartHealth;
+
+	// Setup a GE to modify OutgoingGEs
+	{
+		SKILL_LOG_SCOPE(TEXT("Apply DamageBuff"))
+
+		UGameplayEffect * BuffDmgEffect = Cast<UGameplayEffect>(StaticConstructObject(UGameplayEffect::StaticClass(), GetTransientPackage(), FName(TEXT("DamageBuff"))));
+		BuffDmgEffect->Modifiers.SetNum(1);
+		BuffDmgEffect->Modifiers[0].Magnitude.SetValue(BonusDamageMultiplier);
+		BuffDmgEffect->Modifiers[0].ModifierType = EGameplayMod::OutgoingGE;
+		BuffDmgEffect->Modifiers[0].ModifierOp = EGameplayModOp::Multiplicitive;
+		BuffDmgEffect->Modifiers[0].Attribute.SetUProperty(DamageProperty);
+		BuffDmgEffect->Modifiers[0].OwnedTags.AddTag(IGameplayTagsModule::RequestGameplayTag(FName(TEXT("Damage.Buff"))));
+		BuffDmgEffect->Modifiers[0].RequiredTags.AddTag(IGameplayTagsModule::RequestGameplayTag(FName(TEXT("Damage.Basic"))));
+		BuffDmgEffect->Duration.SetValue(UGameplayEffect::INFINITE_DURATION);
+
+		// Apply to self
+		SourceComponent->ApplyGameplayEffectToTarget(BuffDmgEffect, SourceComponent, 1.f);
+	}
+
+	// attempt to modify the buff but fail because of a 0.f chance to apply
+	{
+		SKILL_LOG_SCOPE(TEXT("Fail to modify DamageBuff"))
+			
+		UGameplayEffect * ModBuffEffect = Cast<UGameplayEffect>(StaticConstructObject(UGameplayEffect::StaticClass(), GetTransientPackage(), FName(TEXT("ModDamageBuff"))));
+		ModBuffEffect->Modifiers.SetNum(1);
+		ModBuffEffect->Modifiers[0].Magnitude.SetValue(ExtraDamageMultiplier);
+		ModBuffEffect->Modifiers[0].ModifierType = EGameplayMod::ActiveGE;
+		ModBuffEffect->Modifiers[0].ModifierOp = EGameplayModOp::Additive;
+		ModBuffEffect->Modifiers[0].Attribute.SetUProperty(DamageProperty);
+		ModBuffEffect->Duration.SetValue(UGameplayEffect::INSTANT_APPLICATION);
+		ModBuffEffect->ChanceToExecuteOnGameplayEffect.SetValue(0.f);
+
+		// Apply to self
+		SourceComponent->ApplyGameplayEffectToTarget(ModBuffEffect, SourceComponent, 1.f);
+	}
+
+	{
+		SKILL_LOG_SCOPE(TEXT("Apply Damage mod"))
+
+		UGameplayEffect * BaseDmgEffect = Cast<UGameplayEffect>(StaticConstructObject(UGameplayEffect::StaticClass(), GetTransientPackage(), FName(TEXT("Damage"))));
+		BaseDmgEffect->Modifiers.SetNum(1);
+		BaseDmgEffect->Modifiers[0].Magnitude.SetValue(DamageValue);
+		BaseDmgEffect->Modifiers[0].ModifierType = EGameplayMod::Attribute;
+		BaseDmgEffect->Modifiers[0].ModifierOp = EGameplayModOp::Additive;
+		BaseDmgEffect->Modifiers[0].Attribute.SetUProperty(DamageProperty);
+		BaseDmgEffect->Modifiers[0].OwnedTags.AddTag(IGameplayTagsModule::RequestGameplayTag(FName(TEXT("Damage.Basic"))));
+		BaseDmgEffect->Duration.SetValue(UGameplayEffect::INSTANT_APPLICATION);
+
+		// Apply to target
+		SourceComponent->ApplyGameplayEffectToTarget(BaseDmgEffect, DestComponent, 1.f);
+
+		float ExpectedValue = StartHealth - (DamageValue * BonusDamageMultiplier);
+		float ActualValue = DestComponent->GetSet<USkillSystemTestAttributeSet>()->Health;
+
+		Test->TestTrue(SKILL_TEST_TEXT("Damage Applied. %.2f == %.2f", ActualValue, ExpectedValue), ActualValue == ExpectedValue);
+	}
+
+	// modify the buff
+	{
+		SKILL_LOG_SCOPE(TEXT("Modify DamageBuff"))
+
+		UGameplayEffect * ModBuffEffect = Cast<UGameplayEffect>(StaticConstructObject(UGameplayEffect::StaticClass(), GetTransientPackage(), FName(TEXT("ModDamageBuff"))));
+		ModBuffEffect->Modifiers.SetNum(1);
+		ModBuffEffect->Modifiers[0].Magnitude.SetValue(ExtraDamageMultiplier);
+		ModBuffEffect->Modifiers[0].ModifierType = EGameplayMod::ActiveGE;
+		ModBuffEffect->Modifiers[0].ModifierOp = EGameplayModOp::Additive;
+		ModBuffEffect->Modifiers[0].Attribute.SetUProperty(DamageProperty);
+		ModBuffEffect->Duration.SetValue(UGameplayEffect::INSTANT_APPLICATION);
+		ModBuffEffect->ChanceToExecuteOnGameplayEffect.SetValue(1.f);
+
+		// Apply to self
+		SourceComponent->ApplyGameplayEffectToTarget(ModBuffEffect, SourceComponent, 1.f);
+	}
+
+	{
+		SKILL_LOG_SCOPE(TEXT("Apply Damage mod"))
+
+		UGameplayEffect * BaseDmgEffect = Cast<UGameplayEffect>(StaticConstructObject(UGameplayEffect::StaticClass(), GetTransientPackage(), FName(TEXT("Damage"))));
+		BaseDmgEffect->Modifiers.SetNum(1);
+		BaseDmgEffect->Modifiers[0].Magnitude.SetValue(DamageValue);
+		BaseDmgEffect->Modifiers[0].ModifierType = EGameplayMod::Attribute;
+		BaseDmgEffect->Modifiers[0].ModifierOp = EGameplayModOp::Additive;
+		BaseDmgEffect->Modifiers[0].Attribute.SetUProperty(DamageProperty);
+		BaseDmgEffect->Modifiers[0].OwnedTags.AddTag(IGameplayTagsModule::RequestGameplayTag(FName(TEXT("Damage.Basic"))));
+		BaseDmgEffect->Duration.SetValue(UGameplayEffect::INSTANT_APPLICATION);
+
+		// Apply to target
+		SourceComponent->ApplyGameplayEffectToTarget(BaseDmgEffect, DestComponent, 1.f);
+
+		float ExpectedValue = StartHealth - (DamageValue * BonusDamageMultiplier) - (DamageValue * (BonusDamageMultiplier + ExtraDamageMultiplier));
+		float ActualValue = DestComponent->GetSet<USkillSystemTestAttributeSet>()->Health;
+
+		Test->TestTrue(SKILL_TEST_TEXT("Damage Applied. %.2f == %.2f", ActualValue, ExpectedValue), ActualValue == ExpectedValue);
+	}
+
+	World->EditorDestroyActor(SourceActor, false);
+	World->EditorDestroyActor(DestActor, false);
+
+	return true;
+}
+
+bool GameplayEffectsTest_ChanceToExecuteOnActiveGEImmunity(UWorld *World, FAutomationTestBase * Test)
+{
+	const float StartHealth = 100.f;
+	const float DamageValue = 5.f;
+	const float BonusDamageMultiplier = 2.f;
+
+	ASkillSystemTestPawn *SourceActor = World->SpawnActor<ASkillSystemTestPawn>();
+	ASkillSystemTestPawn *DestActor = World->SpawnActor<ASkillSystemTestPawn>();
+
+	UProperty *DamageProperty = FindFieldChecked<UProperty>(USkillSystemTestAttributeSet::StaticClass(), GET_MEMBER_NAME_CHECKED(USkillSystemTestAttributeSet, Damage));
+
+	UAttributeComponent * SourceComponent = Cast<UAttributeComponent>(SourceActor->CreateComponentFromTemplate(UAttributeComponent::StaticClass()->GetDefaultObject<UAttributeComponent>(), FString(TEXT("AttributeComponent"))));
+	UAttributeComponent * DestComponent = Cast<UAttributeComponent>(DestActor->CreateComponentFromTemplate(UAttributeComponent::StaticClass()->GetDefaultObject<UAttributeComponent>(), FString(TEXT("AttributeComponent"))));
+	SourceComponent->GetSet<USkillSystemTestAttributeSet>()->Health = StartHealth;
+	DestComponent->GetSet<USkillSystemTestAttributeSet>()->Health = StartHealth;
+
+	// Setup a GE to modify OutgoingGEs
+	{
+		SKILL_LOG_SCOPE(TEXT("Apply DamageBuff"))
+
+		UGameplayEffect * BuffDmgEffect = Cast<UGameplayEffect>(StaticConstructObject(UGameplayEffect::StaticClass(), GetTransientPackage(), FName(TEXT("DamageBuff"))));
+		BuffDmgEffect->Modifiers.SetNum(1);
+		BuffDmgEffect->Modifiers[0].Magnitude.SetValue(BonusDamageMultiplier);
+		BuffDmgEffect->Modifiers[0].ModifierType = EGameplayMod::OutgoingGE;
+		BuffDmgEffect->Modifiers[0].ModifierOp = EGameplayModOp::Multiplicitive;
+		BuffDmgEffect->Modifiers[0].Attribute.SetUProperty(DamageProperty);
+		BuffDmgEffect->Modifiers[0].OwnedTags.AddTag(IGameplayTagsModule::RequestGameplayTag(FName(TEXT("Damage.Buff"))));
+		BuffDmgEffect->Modifiers[0].RequiredTags.AddTag(IGameplayTagsModule::RequestGameplayTag(FName(TEXT("Damage.Basic"))));
+		BuffDmgEffect->Duration.SetValue(UGameplayEffect::INFINITE_DURATION);
+
+		// Apply to self
+		SourceComponent->ApplyGameplayEffectToTarget(BuffDmgEffect, SourceComponent, 1.f);
+	}
+
+	// attempt to remove the buff but fail because of a 0.f chance to apply
+	{
+		SKILL_LOG_SCOPE(TEXT("Fail to remove DamageBuff"))
+
+		UGameplayEffect * ModBuffEffect = Cast<UGameplayEffect>(StaticConstructObject(UGameplayEffect::StaticClass(), GetTransientPackage(), FName(TEXT("ModDamageBuff"))));
+		ModBuffEffect->GameplayEffectRequiredTags.AddTag(IGameplayTagsModule::RequestGameplayTag(FName(TEXT("Damage.Buff"))));
+		ModBuffEffect->AppliesImmunityTo = EGameplayImmunity::ActiveGE;
+		ModBuffEffect->Duration.SetValue(UGameplayEffect::INSTANT_APPLICATION);
+		ModBuffEffect->ChanceToExecuteOnGameplayEffect.SetValue(0.f);
+
+		// Apply to self
+		SourceComponent->ApplyGameplayEffectToTarget(ModBuffEffect, SourceComponent, 1.f);
+	}
+
+	{
+		SKILL_LOG_SCOPE(TEXT("Apply Damage mod"))
+
+		UGameplayEffect * BaseDmgEffect = Cast<UGameplayEffect>(StaticConstructObject(UGameplayEffect::StaticClass(), GetTransientPackage(), FName(TEXT("Damage"))));
+		BaseDmgEffect->Modifiers.SetNum(1);
+		BaseDmgEffect->Modifiers[0].Magnitude.SetValue(DamageValue);
+		BaseDmgEffect->Modifiers[0].ModifierType = EGameplayMod::Attribute;
+		BaseDmgEffect->Modifiers[0].ModifierOp = EGameplayModOp::Additive;
+		BaseDmgEffect->Modifiers[0].Attribute.SetUProperty(DamageProperty);
+		BaseDmgEffect->Modifiers[0].OwnedTags.AddTag(IGameplayTagsModule::RequestGameplayTag(FName(TEXT("Damage.Basic"))));
+		BaseDmgEffect->Duration.SetValue(UGameplayEffect::INSTANT_APPLICATION);
+
+		// Apply to target
+		SourceComponent->ApplyGameplayEffectToTarget(BaseDmgEffect, DestComponent, 1.f);
+
+		float ExpectedValue = StartHealth - (DamageValue * BonusDamageMultiplier);
+		float ActualValue = DestComponent->GetSet<USkillSystemTestAttributeSet>()->Health;
+
+		Test->TestTrue(SKILL_TEST_TEXT("Damage Applied. %.2f == %.2f", ActualValue, ExpectedValue), ActualValue == ExpectedValue);
+	}
+
+	// modify the buff
+	{
+		SKILL_LOG_SCOPE(TEXT("Remove DamageBuff"))
+
+		UGameplayEffect * ModBuffEffect = Cast<UGameplayEffect>(StaticConstructObject(UGameplayEffect::StaticClass(), GetTransientPackage(), FName(TEXT("ModDamageBuff"))));
+		ModBuffEffect->GameplayEffectRequiredTags.AddTag(IGameplayTagsModule::RequestGameplayTag(FName(TEXT("Damage.Buff"))));
+		ModBuffEffect->AppliesImmunityTo = EGameplayImmunity::ActiveGE;
+		ModBuffEffect->Duration.SetValue(UGameplayEffect::INSTANT_APPLICATION);
+		ModBuffEffect->ChanceToExecuteOnGameplayEffect.SetValue(1.f);
+
+		// Apply to self
+		SourceComponent->ApplyGameplayEffectToTarget(ModBuffEffect, SourceComponent, 1.f);
+	}
+
+	{
+		SKILL_LOG_SCOPE(TEXT("Apply Damage mod"))
+
+		UGameplayEffect * BaseDmgEffect = Cast<UGameplayEffect>(StaticConstructObject(UGameplayEffect::StaticClass(), GetTransientPackage(), FName(TEXT("Damage"))));
+		BaseDmgEffect->Modifiers.SetNum(1);
+		BaseDmgEffect->Modifiers[0].Magnitude.SetValue(DamageValue);
+		BaseDmgEffect->Modifiers[0].ModifierType = EGameplayMod::Attribute;
+		BaseDmgEffect->Modifiers[0].ModifierOp = EGameplayModOp::Additive;
+		BaseDmgEffect->Modifiers[0].Attribute.SetUProperty(DamageProperty);
+		BaseDmgEffect->Modifiers[0].OwnedTags.AddTag(IGameplayTagsModule::RequestGameplayTag(FName(TEXT("Damage.Basic"))));
+		BaseDmgEffect->Duration.SetValue(UGameplayEffect::INSTANT_APPLICATION);
+
+		// Apply to target
+		SourceComponent->ApplyGameplayEffectToTarget(BaseDmgEffect, DestComponent, 1.f);
+
+		float ExpectedValue = StartHealth - (DamageValue * BonusDamageMultiplier) - (DamageValue);
+		float ActualValue = DestComponent->GetSet<USkillSystemTestAttributeSet>()->Health;
+
+		Test->TestTrue(SKILL_TEST_TEXT("Damage Applied. %.2f == %.2f", ActualValue, ExpectedValue), ActualValue == ExpectedValue);
+	}
+
+	World->EditorDestroyActor(SourceActor, false);
+	World->EditorDestroyActor(DestActor, false);
+
+	return true;
+}
+
+bool GameplayEffectsTest_ChanceToExecuteOnOutgoingGEMod(UWorld *World, FAutomationTestBase * Test)
+{
+	const float StartHealth = 100.f;
+	const float DamageValue = 5.f;
+	const float BonusDamageMultiplier = 2.f;
+
+	ASkillSystemTestPawn *SourceActor = World->SpawnActor<ASkillSystemTestPawn>();
+	ASkillSystemTestPawn *DestActor = World->SpawnActor<ASkillSystemTestPawn>();
+
+	UProperty *DamageProperty = FindFieldChecked<UProperty>(USkillSystemTestAttributeSet::StaticClass(), GET_MEMBER_NAME_CHECKED(USkillSystemTestAttributeSet, Damage));
+
+	UAttributeComponent * SourceComponent = Cast<UAttributeComponent>(SourceActor->CreateComponentFromTemplate(UAttributeComponent::StaticClass()->GetDefaultObject<UAttributeComponent>(), FString(TEXT("AttributeComponent"))));
+	UAttributeComponent * DestComponent = Cast<UAttributeComponent>(DestActor->CreateComponentFromTemplate(UAttributeComponent::StaticClass()->GetDefaultObject<UAttributeComponent>(), FString(TEXT("AttributeComponent"))));
+	SourceComponent->GetSet<USkillSystemTestAttributeSet>()->Health = StartHealth;
+	DestComponent->GetSet<USkillSystemTestAttributeSet>()->Health = StartHealth;
+
+	// Setup a GE to modify OutgoingGEs
+	// this GE won't do anything because the chance to execute is zero
+	{
+		SKILL_LOG_SCOPE(TEXT("Apply useless DamageBuff"))
+
+		UGameplayEffect * BuffDmgEffect = Cast<UGameplayEffect>(StaticConstructObject(UGameplayEffect::StaticClass(), GetTransientPackage(), FName(TEXT("DamageBuff"))));
+		BuffDmgEffect->Modifiers.SetNum(1);
+		BuffDmgEffect->Modifiers[0].Magnitude.SetValue(BonusDamageMultiplier);
+		BuffDmgEffect->Modifiers[0].ModifierType = EGameplayMod::OutgoingGE;
+		BuffDmgEffect->Modifiers[0].ModifierOp = EGameplayModOp::Multiplicitive;
+		BuffDmgEffect->Modifiers[0].Attribute.SetUProperty(DamageProperty);
+		BuffDmgEffect->Modifiers[0].OwnedTags.AddTag(IGameplayTagsModule::RequestGameplayTag(FName(TEXT("Damage.Buff"))));
+		BuffDmgEffect->Modifiers[0].RequiredTags.AddTag(IGameplayTagsModule::RequestGameplayTag(FName(TEXT("Damage.Basic"))));
+		BuffDmgEffect->Duration.SetValue(UGameplayEffect::INFINITE_DURATION);
+		BuffDmgEffect->ChanceToExecuteOnGameplayEffect.SetValue(0.f);
+
+		// Apply to self
+		SourceComponent->ApplyGameplayEffectToTarget(BuffDmgEffect, SourceComponent, 1.f);
+	}
+
+	{
+		SKILL_LOG_SCOPE(TEXT("Apply Damage mod"))
+
+		UGameplayEffect * BaseDmgEffect = Cast<UGameplayEffect>(StaticConstructObject(UGameplayEffect::StaticClass(), GetTransientPackage(), FName(TEXT("Damage"))));
+		BaseDmgEffect->Modifiers.SetNum(1);
+		BaseDmgEffect->Modifiers[0].Magnitude.SetValue(DamageValue);
+		BaseDmgEffect->Modifiers[0].ModifierType = EGameplayMod::Attribute;
+		BaseDmgEffect->Modifiers[0].ModifierOp = EGameplayModOp::Additive;
+		BaseDmgEffect->Modifiers[0].Attribute.SetUProperty(DamageProperty);
+		BaseDmgEffect->Modifiers[0].OwnedTags.AddTag(IGameplayTagsModule::RequestGameplayTag(FName(TEXT("Damage.Basic"))));
+		BaseDmgEffect->Duration.SetValue(UGameplayEffect::INSTANT_APPLICATION);
+
+		// Apply to target
+		SourceComponent->ApplyGameplayEffectToTarget(BaseDmgEffect, DestComponent, 1.f);
+
+		float ExpectedValue = StartHealth - DamageValue;
+		float ActualValue = DestComponent->GetSet<USkillSystemTestAttributeSet>()->Health;
+
+		Test->TestTrue(SKILL_TEST_TEXT("Damage Applied. %.2f == %.2f", ActualValue, ExpectedValue), ActualValue == ExpectedValue);
+	}
+
+	// Setup a GE to modify OutgoingGEs
+	// this GE will always execute work because the chance to execute is one
+	{
+		SKILL_LOG_SCOPE(TEXT("Apply DamageBuff"))
+
+		UGameplayEffect * BuffDmgEffect = Cast<UGameplayEffect>(StaticConstructObject(UGameplayEffect::StaticClass(), GetTransientPackage(), FName(TEXT("DamageBuff"))));
+		BuffDmgEffect->Modifiers.SetNum(1);
+		BuffDmgEffect->Modifiers[0].Magnitude.SetValue(BonusDamageMultiplier);
+		BuffDmgEffect->Modifiers[0].ModifierType = EGameplayMod::OutgoingGE;
+		BuffDmgEffect->Modifiers[0].ModifierOp = EGameplayModOp::Multiplicitive;
+		BuffDmgEffect->Modifiers[0].Attribute.SetUProperty(DamageProperty);
+		BuffDmgEffect->Modifiers[0].OwnedTags.AddTag(IGameplayTagsModule::RequestGameplayTag(FName(TEXT("Damage.Buff"))));
+		BuffDmgEffect->Modifiers[0].RequiredTags.AddTag(IGameplayTagsModule::RequestGameplayTag(FName(TEXT("Damage.Basic"))));
+		BuffDmgEffect->Duration.SetValue(UGameplayEffect::INFINITE_DURATION);
+		BuffDmgEffect->ChanceToExecuteOnGameplayEffect.SetValue(1.f);
+
+		// Apply to self
+		SourceComponent->ApplyGameplayEffectToTarget(BuffDmgEffect, SourceComponent, 1.f);
+	}
+
+	{
+		SKILL_LOG_SCOPE(TEXT("Apply Damage mod"))
+
+		UGameplayEffect * BaseDmgEffect = Cast<UGameplayEffect>(StaticConstructObject(UGameplayEffect::StaticClass(), GetTransientPackage(), FName(TEXT("Damage"))));
+		BaseDmgEffect->Modifiers.SetNum(1);
+		BaseDmgEffect->Modifiers[0].Magnitude.SetValue(DamageValue);
+		BaseDmgEffect->Modifiers[0].ModifierType = EGameplayMod::Attribute;
+		BaseDmgEffect->Modifiers[0].ModifierOp = EGameplayModOp::Additive;
+		BaseDmgEffect->Modifiers[0].Attribute.SetUProperty(DamageProperty);
+		BaseDmgEffect->Modifiers[0].OwnedTags.AddTag(IGameplayTagsModule::RequestGameplayTag(FName(TEXT("Damage.Basic"))));
+		BaseDmgEffect->Duration.SetValue(UGameplayEffect::INSTANT_APPLICATION);
+
+		// Apply to target
+		SourceComponent->ApplyGameplayEffectToTarget(BaseDmgEffect, DestComponent, 1.f);
+
+		float ExpectedValue = StartHealth - DamageValue - (DamageValue * BonusDamageMultiplier);
+		float ActualValue = DestComponent->GetSet<USkillSystemTestAttributeSet>()->Health;
+
+		Test->TestTrue(SKILL_TEST_TEXT("Damage Applied. %.2f == %.2f", ActualValue, ExpectedValue), ActualValue == ExpectedValue);
+	}
+
+	World->EditorDestroyActor(SourceActor, false);
+	World->EditorDestroyActor(DestActor, false);
+
+	return true;
+}
+
+bool GameplayEffectsTest_ChanceToExecuteOnOutgoingGEImmunity(UWorld *World, FAutomationTestBase * Test)
+{
+	const float StartHealth = 100.f;
+	const float DamageValue = 5.f;
+
+	ASkillSystemTestPawn *SourceActor = World->SpawnActor<ASkillSystemTestPawn>();
+	ASkillSystemTestPawn *DestActor = World->SpawnActor<ASkillSystemTestPawn>();
+
+	UProperty *DamageProperty = FindFieldChecked<UProperty>(USkillSystemTestAttributeSet::StaticClass(), GET_MEMBER_NAME_CHECKED(USkillSystemTestAttributeSet, Damage));
+
+	UAttributeComponent * SourceComponent = Cast<UAttributeComponent>(SourceActor->CreateComponentFromTemplate(UAttributeComponent::StaticClass()->GetDefaultObject<UAttributeComponent>(), FString(TEXT("AttributeComponent"))));
+	UAttributeComponent * DestComponent = Cast<UAttributeComponent>(DestActor->CreateComponentFromTemplate(UAttributeComponent::StaticClass()->GetDefaultObject<UAttributeComponent>(), FString(TEXT("AttributeComponent"))));
+	SourceComponent->GetSet<USkillSystemTestAttributeSet>()->Health = StartHealth;
+	DestComponent->GetSet<USkillSystemTestAttributeSet>()->Health = StartHealth;
+
+
+	// attempt to prevent outgoing damage but fail because of a 0.f chance to apply
+	{
+		SKILL_LOG_SCOPE(TEXT("Fail to remove DamageBuff"))
+
+		UGameplayEffect * ModBuffEffect = Cast<UGameplayEffect>(StaticConstructObject(UGameplayEffect::StaticClass(), GetTransientPackage(), FName(TEXT("ModDamageBuff"))));
+		ModBuffEffect->GameplayEffectRequiredTags.AddTag(IGameplayTagsModule::RequestGameplayTag(FName(TEXT("Damage.Basic"))));
+		ModBuffEffect->AppliesImmunityTo = EGameplayImmunity::OutgoingGE;
+		ModBuffEffect->Duration.SetValue(UGameplayEffect::INFINITE_DURATION);
+		ModBuffEffect->ChanceToExecuteOnGameplayEffect.SetValue(0.f);
+
+		// Apply to self
+		SourceComponent->ApplyGameplayEffectToTarget(ModBuffEffect, SourceComponent, 1.f);
+	}
+
+	{
+		SKILL_LOG_SCOPE(TEXT("Apply Damage mod"))
+
+		UGameplayEffect * BaseDmgEffect = Cast<UGameplayEffect>(StaticConstructObject(UGameplayEffect::StaticClass(), GetTransientPackage(), FName(TEXT("Damage"))));
+		BaseDmgEffect->Modifiers.SetNum(1);
+		BaseDmgEffect->Modifiers[0].Magnitude.SetValue(DamageValue);
+		BaseDmgEffect->Modifiers[0].ModifierType = EGameplayMod::Attribute;
+		BaseDmgEffect->Modifiers[0].ModifierOp = EGameplayModOp::Additive;
+		BaseDmgEffect->Modifiers[0].Attribute.SetUProperty(DamageProperty);
+		BaseDmgEffect->Modifiers[0].OwnedTags.AddTag(IGameplayTagsModule::RequestGameplayTag(FName(TEXT("Damage.Basic"))));
+		BaseDmgEffect->Duration.SetValue(UGameplayEffect::INSTANT_APPLICATION);
+
+		// Apply to target
+		SourceComponent->ApplyGameplayEffectToTarget(BaseDmgEffect, DestComponent, 1.f);
+
+		float ExpectedValue = StartHealth - DamageValue;
+		float ActualValue = DestComponent->GetSet<USkillSystemTestAttributeSet>()->Health;
+
+		Test->TestTrue(SKILL_TEST_TEXT("Damage Applied. %.2f == %.2f", ActualValue, ExpectedValue), ActualValue == ExpectedValue);
+	}
+
+	// prevent outgoing damage buff
+	{
+		SKILL_LOG_SCOPE(TEXT("Fail to remove DamageBuff"))
+
+		UGameplayEffect * ModBuffEffect = Cast<UGameplayEffect>(StaticConstructObject(UGameplayEffect::StaticClass(), GetTransientPackage(), FName(TEXT("ModDamageBuff"))));
+		ModBuffEffect->GameplayEffectRequiredTags.AddTag(IGameplayTagsModule::RequestGameplayTag(FName(TEXT("Damage.Basic"))));
+		ModBuffEffect->AppliesImmunityTo = EGameplayImmunity::OutgoingGE;
+		ModBuffEffect->Duration.SetValue(UGameplayEffect::INFINITE_DURATION);
+		ModBuffEffect->ChanceToExecuteOnGameplayEffect.SetValue(1.f);
+
+		// Apply to self
+		SourceComponent->ApplyGameplayEffectToTarget(ModBuffEffect, SourceComponent, 1.f);
+	}
+
+	{
+		SKILL_LOG_SCOPE(TEXT("Apply Damage mod"))
+
+		UGameplayEffect * BaseDmgEffect = Cast<UGameplayEffect>(StaticConstructObject(UGameplayEffect::StaticClass(), GetTransientPackage(), FName(TEXT("Damage"))));
+		BaseDmgEffect->Modifiers.SetNum(1);
+		BaseDmgEffect->Modifiers[0].Magnitude.SetValue(DamageValue);
+		BaseDmgEffect->Modifiers[0].ModifierType = EGameplayMod::Attribute;
+		BaseDmgEffect->Modifiers[0].ModifierOp = EGameplayModOp::Additive;
+		BaseDmgEffect->Modifiers[0].Attribute.SetUProperty(DamageProperty);
+		BaseDmgEffect->Modifiers[0].OwnedTags.AddTag(IGameplayTagsModule::RequestGameplayTag(FName(TEXT("Damage.Basic"))));
+		BaseDmgEffect->Duration.SetValue(UGameplayEffect::INSTANT_APPLICATION);
+
+		// Apply to target
+		SourceComponent->ApplyGameplayEffectToTarget(BaseDmgEffect, DestComponent, 1.f);
+
+		float ExpectedValue = StartHealth - DamageValue;
+		float ActualValue = DestComponent->GetSet<USkillSystemTestAttributeSet>()->Health;
+
+		Test->TestTrue(SKILL_TEST_TEXT("Damage Applied. %.2f == %.2f", ActualValue, ExpectedValue), ActualValue == ExpectedValue);
+	}
+
+	World->EditorDestroyActor(SourceActor, false);
+	World->EditorDestroyActor(DestActor, false);
+
+	return true;
+}
+
+bool GameplayEffectsTest_ChanceToExecuteOnIncomingGEMod(UWorld *World, FAutomationTestBase * Test)
+{
+	const float StartHealth = 100.f;
+	const float DamageValue = 5.f;
+	const float BonusDamageMultiplier = 2.f;
+
+	ASkillSystemTestPawn *SourceActor = World->SpawnActor<ASkillSystemTestPawn>();
+	ASkillSystemTestPawn *DestActor = World->SpawnActor<ASkillSystemTestPawn>();
+
+	UProperty *DamageProperty = FindFieldChecked<UProperty>(USkillSystemTestAttributeSet::StaticClass(), GET_MEMBER_NAME_CHECKED(USkillSystemTestAttributeSet, Damage));
+
+	UAttributeComponent * SourceComponent = Cast<UAttributeComponent>(SourceActor->CreateComponentFromTemplate(UAttributeComponent::StaticClass()->GetDefaultObject<UAttributeComponent>(), FString(TEXT("AttributeComponent"))));
+	UAttributeComponent * DestComponent = Cast<UAttributeComponent>(DestActor->CreateComponentFromTemplate(UAttributeComponent::StaticClass()->GetDefaultObject<UAttributeComponent>(), FString(TEXT("AttributeComponent"))));
+	SourceComponent->GetSet<USkillSystemTestAttributeSet>()->Health = StartHealth;
+	DestComponent->GetSet<USkillSystemTestAttributeSet>()->Health = StartHealth;
+
+	// Setup a GE to modify OutgoingGEs
+	// this GE won't do anything because the chance to execute is zero
+	{
+		SKILL_LOG_SCOPE(TEXT("Apply useless DamageBuff"))
+
+		UGameplayEffect * BuffDmgEffect = Cast<UGameplayEffect>(StaticConstructObject(UGameplayEffect::StaticClass(), GetTransientPackage(), FName(TEXT("DamageBuff"))));
+		BuffDmgEffect->Modifiers.SetNum(1);
+		BuffDmgEffect->Modifiers[0].Magnitude.SetValue(BonusDamageMultiplier);
+		BuffDmgEffect->Modifiers[0].ModifierType = EGameplayMod::IncomingGE;
+		BuffDmgEffect->Modifiers[0].ModifierOp = EGameplayModOp::Multiplicitive;
+		BuffDmgEffect->Modifiers[0].Attribute.SetUProperty(DamageProperty);
+		BuffDmgEffect->Modifiers[0].OwnedTags.AddTag(IGameplayTagsModule::RequestGameplayTag(FName(TEXT("Damage.Buff"))));
+		BuffDmgEffect->Modifiers[0].RequiredTags.AddTag(IGameplayTagsModule::RequestGameplayTag(FName(TEXT("Damage.Basic"))));
+		BuffDmgEffect->Duration.SetValue(UGameplayEffect::INFINITE_DURATION);
+		BuffDmgEffect->ChanceToExecuteOnGameplayEffect.SetValue(0.f);
+
+		// Apply to self
+		SourceComponent->ApplyGameplayEffectToTarget(BuffDmgEffect, DestComponent, 1.f);
+	}
+
+	{
+		SKILL_LOG_SCOPE(TEXT("Apply Damage mod"))
+
+		UGameplayEffect * BaseDmgEffect = Cast<UGameplayEffect>(StaticConstructObject(UGameplayEffect::StaticClass(), GetTransientPackage(), FName(TEXT("Damage"))));
+		BaseDmgEffect->Modifiers.SetNum(1);
+		BaseDmgEffect->Modifiers[0].Magnitude.SetValue(DamageValue);
+		BaseDmgEffect->Modifiers[0].ModifierType = EGameplayMod::Attribute;
+		BaseDmgEffect->Modifiers[0].ModifierOp = EGameplayModOp::Additive;
+		BaseDmgEffect->Modifiers[0].Attribute.SetUProperty(DamageProperty);
+		BaseDmgEffect->Modifiers[0].OwnedTags.AddTag(IGameplayTagsModule::RequestGameplayTag(FName(TEXT("Damage.Basic"))));
+		BaseDmgEffect->Duration.SetValue(UGameplayEffect::INSTANT_APPLICATION);
+
+		// Apply to target
+		SourceComponent->ApplyGameplayEffectToTarget(BaseDmgEffect, DestComponent, 1.f);
+
+		float ExpectedValue = StartHealth - DamageValue;
+		float ActualValue = DestComponent->GetSet<USkillSystemTestAttributeSet>()->Health;
+
+		Test->TestTrue(SKILL_TEST_TEXT("Damage Applied. %.2f == %.2f", ActualValue, ExpectedValue), ActualValue == ExpectedValue);
+	}
+
+	// Setup a GE to modify OutgoingGEs
+	// this GE will always execute work because the chance to execute is one
+	{
+		SKILL_LOG_SCOPE(TEXT("Apply DamageBuff"))
+
+		UGameplayEffect * BuffDmgEffect = Cast<UGameplayEffect>(StaticConstructObject(UGameplayEffect::StaticClass(), GetTransientPackage(), FName(TEXT("DamageBuff"))));
+		BuffDmgEffect->Modifiers.SetNum(1);
+		BuffDmgEffect->Modifiers[0].Magnitude.SetValue(BonusDamageMultiplier);
+		BuffDmgEffect->Modifiers[0].ModifierType = EGameplayMod::IncomingGE;
+		BuffDmgEffect->Modifiers[0].ModifierOp = EGameplayModOp::Multiplicitive;
+		BuffDmgEffect->Modifiers[0].Attribute.SetUProperty(DamageProperty);
+		BuffDmgEffect->Modifiers[0].OwnedTags.AddTag(IGameplayTagsModule::RequestGameplayTag(FName(TEXT("Damage.Buff"))));
+		BuffDmgEffect->Modifiers[0].RequiredTags.AddTag(IGameplayTagsModule::RequestGameplayTag(FName(TEXT("Damage.Basic"))));
+		BuffDmgEffect->Duration.SetValue(UGameplayEffect::INFINITE_DURATION);
+		BuffDmgEffect->ChanceToExecuteOnGameplayEffect.SetValue(1.f);
+
+		// Apply to self
+		SourceComponent->ApplyGameplayEffectToTarget(BuffDmgEffect, DestComponent, 1.f);
+	}
+
+	{
+		SKILL_LOG_SCOPE(TEXT("Apply Damage mod"))
+
+		UGameplayEffect * BaseDmgEffect = Cast<UGameplayEffect>(StaticConstructObject(UGameplayEffect::StaticClass(), GetTransientPackage(), FName(TEXT("Damage"))));
+		BaseDmgEffect->Modifiers.SetNum(1);
+		BaseDmgEffect->Modifiers[0].Magnitude.SetValue(DamageValue);
+		BaseDmgEffect->Modifiers[0].ModifierType = EGameplayMod::Attribute;
+		BaseDmgEffect->Modifiers[0].ModifierOp = EGameplayModOp::Additive;
+		BaseDmgEffect->Modifiers[0].Attribute.SetUProperty(DamageProperty);
+		BaseDmgEffect->Modifiers[0].OwnedTags.AddTag(IGameplayTagsModule::RequestGameplayTag(FName(TEXT("Damage.Basic"))));
+		BaseDmgEffect->Duration.SetValue(UGameplayEffect::INSTANT_APPLICATION);
+
+		// Apply to target
+		SourceComponent->ApplyGameplayEffectToTarget(BaseDmgEffect, DestComponent, 1.f);
+
+		float ExpectedValue = StartHealth - DamageValue - (DamageValue * BonusDamageMultiplier);
+		float ActualValue = DestComponent->GetSet<USkillSystemTestAttributeSet>()->Health;
+
+		Test->TestTrue(SKILL_TEST_TEXT("Damage Applied. %.2f == %.2f", ActualValue, ExpectedValue), ActualValue == ExpectedValue);
+	}
+
+	World->EditorDestroyActor(SourceActor, false);
+	World->EditorDestroyActor(DestActor, false);
+
+	return true;
+}
+
+bool GameplayEffectsTest_ChanceToExecuteOnIncomingGEImmunity(UWorld *World, FAutomationTestBase * Test)
+{
+	const float StartHealth = 100.f;
+	const float DamageValue = 5.f;
+
+	ASkillSystemTestPawn *SourceActor = World->SpawnActor<ASkillSystemTestPawn>();
+	ASkillSystemTestPawn *DestActor = World->SpawnActor<ASkillSystemTestPawn>();
+
+	UProperty *DamageProperty = FindFieldChecked<UProperty>(USkillSystemTestAttributeSet::StaticClass(), GET_MEMBER_NAME_CHECKED(USkillSystemTestAttributeSet, Damage));
+
+	UAttributeComponent * SourceComponent = Cast<UAttributeComponent>(SourceActor->CreateComponentFromTemplate(UAttributeComponent::StaticClass()->GetDefaultObject<UAttributeComponent>(), FString(TEXT("AttributeComponent"))));
+	UAttributeComponent * DestComponent = Cast<UAttributeComponent>(DestActor->CreateComponentFromTemplate(UAttributeComponent::StaticClass()->GetDefaultObject<UAttributeComponent>(), FString(TEXT("AttributeComponent"))));
+	SourceComponent->GetSet<USkillSystemTestAttributeSet>()->Health = StartHealth;
+	DestComponent->GetSet<USkillSystemTestAttributeSet>()->Health = StartHealth;
+
+
+	// attempt to prevent outgoing damage but fail because of a 0.f chance to apply
+	{
+		SKILL_LOG_SCOPE(TEXT("Fail to remove DamageBuff"))
+
+		UGameplayEffect * ModBuffEffect = Cast<UGameplayEffect>(StaticConstructObject(UGameplayEffect::StaticClass(), GetTransientPackage(), FName(TEXT("ModDamageBuff"))));
+		ModBuffEffect->GameplayEffectRequiredTags.AddTag(IGameplayTagsModule::RequestGameplayTag(FName(TEXT("Damage.Basic"))));
+		ModBuffEffect->AppliesImmunityTo = EGameplayImmunity::IncomingGE;
+		ModBuffEffect->Duration.SetValue(UGameplayEffect::INFINITE_DURATION);
+		ModBuffEffect->ChanceToExecuteOnGameplayEffect.SetValue(0.f);
+
+		// Apply to self
+		SourceComponent->ApplyGameplayEffectToTarget(ModBuffEffect, DestComponent, 1.f);
+	}
+
+	{
+		SKILL_LOG_SCOPE(TEXT("Apply Damage mod"))
+
+		UGameplayEffect * BaseDmgEffect = Cast<UGameplayEffect>(StaticConstructObject(UGameplayEffect::StaticClass(), GetTransientPackage(), FName(TEXT("Damage"))));
+		BaseDmgEffect->Modifiers.SetNum(1);
+		BaseDmgEffect->Modifiers[0].Magnitude.SetValue(DamageValue);
+		BaseDmgEffect->Modifiers[0].ModifierType = EGameplayMod::Attribute;
+		BaseDmgEffect->Modifiers[0].ModifierOp = EGameplayModOp::Additive;
+		BaseDmgEffect->Modifiers[0].Attribute.SetUProperty(DamageProperty);
+		BaseDmgEffect->Modifiers[0].OwnedTags.AddTag(IGameplayTagsModule::RequestGameplayTag(FName(TEXT("Damage.Basic"))));
+		BaseDmgEffect->Duration.SetValue(UGameplayEffect::INSTANT_APPLICATION);
+
+		// Apply to target
+		SourceComponent->ApplyGameplayEffectToTarget(BaseDmgEffect, DestComponent, 1.f);
+
+		float ExpectedValue = StartHealth - DamageValue;
+		float ActualValue = DestComponent->GetSet<USkillSystemTestAttributeSet>()->Health;
+
+		Test->TestTrue(SKILL_TEST_TEXT("Damage Applied. %.2f == %.2f", ActualValue, ExpectedValue), ActualValue == ExpectedValue);
+	}
+
+	// prevent outgoing damage buff
+	{
+		SKILL_LOG_SCOPE(TEXT("Fail to remove DamageBuff"))
+
+		UGameplayEffect * ModBuffEffect = Cast<UGameplayEffect>(StaticConstructObject(UGameplayEffect::StaticClass(), GetTransientPackage(), FName(TEXT("ModDamageBuff"))));
+		ModBuffEffect->GameplayEffectRequiredTags.AddTag(IGameplayTagsModule::RequestGameplayTag(FName(TEXT("Damage.Basic"))));
+		ModBuffEffect->AppliesImmunityTo = EGameplayImmunity::IncomingGE;
+		ModBuffEffect->Duration.SetValue(UGameplayEffect::INFINITE_DURATION);
+		ModBuffEffect->ChanceToExecuteOnGameplayEffect.SetValue(1.f);
+
+		// Apply to self
+		SourceComponent->ApplyGameplayEffectToTarget(ModBuffEffect, DestComponent, 1.f);
+	}
+
+	{
+		SKILL_LOG_SCOPE(TEXT("Apply Damage mod"))
+
+		UGameplayEffect * BaseDmgEffect = Cast<UGameplayEffect>(StaticConstructObject(UGameplayEffect::StaticClass(), GetTransientPackage(), FName(TEXT("Damage"))));
+		BaseDmgEffect->Modifiers.SetNum(1);
+		BaseDmgEffect->Modifiers[0].Magnitude.SetValue(DamageValue);
+		BaseDmgEffect->Modifiers[0].ModifierType = EGameplayMod::Attribute;
+		BaseDmgEffect->Modifiers[0].ModifierOp = EGameplayModOp::Additive;
+		BaseDmgEffect->Modifiers[0].Attribute.SetUProperty(DamageProperty);
+		BaseDmgEffect->Modifiers[0].OwnedTags.AddTag(IGameplayTagsModule::RequestGameplayTag(FName(TEXT("Damage.Basic"))));
+		BaseDmgEffect->Duration.SetValue(UGameplayEffect::INSTANT_APPLICATION);
+
+		// Apply to target
+		SourceComponent->ApplyGameplayEffectToTarget(BaseDmgEffect, DestComponent, 1.f);
+
+		float ExpectedValue = StartHealth - DamageValue;
+		float ActualValue = DestComponent->GetSet<USkillSystemTestAttributeSet>()->Health;
+
+		Test->TestTrue(SKILL_TEST_TEXT("Damage Applied. %.2f == %.2f", ActualValue, ExpectedValue), ActualValue == ExpectedValue);
+	}
+
+	World->EditorDestroyActor(SourceActor, false);
+	World->EditorDestroyActor(DestActor, false);
+
+	return true;
+}
+
+bool GameplayEffectsTest_ModifyChanceToApplyToTarget(UWorld *World, FAutomationTestBase * Test)
+{
+	const float StartHealth = 100.f;
+	const float DamageValue = 5.f;
+
+	ASkillSystemTestPawn *SourceActor = World->SpawnActor<ASkillSystemTestPawn>();
+	ASkillSystemTestPawn *DestActor = World->SpawnActor<ASkillSystemTestPawn>();
+
+	UProperty *DamageProperty = FindFieldChecked<UProperty>(USkillSystemTestAttributeSet::StaticClass(), GET_MEMBER_NAME_CHECKED(USkillSystemTestAttributeSet, Damage));
+
+	UAttributeComponent * SourceComponent = Cast<UAttributeComponent>(SourceActor->CreateComponentFromTemplate(UAttributeComponent::StaticClass()->GetDefaultObject<UAttributeComponent>(), FString(TEXT("AttributeComponent"))));
+	UAttributeComponent * DestComponent = Cast<UAttributeComponent>(DestActor->CreateComponentFromTemplate(UAttributeComponent::StaticClass()->GetDefaultObject<UAttributeComponent>(), FString(TEXT("AttributeComponent"))));
+	SourceComponent->GetSet<USkillSystemTestAttributeSet>()->Health = StartHealth;
+	DestComponent->GetSet<USkillSystemTestAttributeSet>()->Health = StartHealth;
+
+	{
+		SKILL_LOG_SCOPE(TEXT("Apply boost to chance to apply"))
+
+		UGameplayEffect * BaseEffect = Cast<UGameplayEffect>(StaticConstructObject(UGameplayEffect::StaticClass(), GetTransientPackage(), FName(TEXT("ChanceToApply"))));
+		BaseEffect->Modifiers.SetNum(1);
+		BaseEffect->Modifiers[0].Magnitude.SetValue(1.f);
+		BaseEffect->Modifiers[0].ModifierType = EGameplayMod::OutgoingGE;
+		BaseEffect->Modifiers[0].ModifierOp = EGameplayModOp::Additive;
+		BaseEffect->Modifiers[0].EffectType = EGameplayModEffect::ChanceApplyTarget;
+		BaseEffect->Duration.SetValue(UGameplayEffect::INFINITE_DURATION);
+
+		// Apply to self
+		SourceComponent->ApplyGameplayEffectToTarget(BaseEffect, SourceComponent, 0.f);
+	}
+
+	{
+		SKILL_LOG_SCOPE(TEXT("Apply Damage mod"))
+
+		UGameplayEffect * BaseDmgEffect = Cast<UGameplayEffect>(StaticConstructObject(UGameplayEffect::StaticClass(), GetTransientPackage(), FName(TEXT("Damage"))));
+		BaseDmgEffect->Modifiers.SetNum(1);
+		BaseDmgEffect->Modifiers[0].Magnitude.SetValue(DamageValue);
+		BaseDmgEffect->Modifiers[0].ModifierType = EGameplayMod::Attribute;
+		BaseDmgEffect->Modifiers[0].ModifierOp = EGameplayModOp::Additive;
+		BaseDmgEffect->Modifiers[0].Attribute.SetUProperty(DamageProperty);
+		BaseDmgEffect->Modifiers[0].OwnedTags.AddTag(IGameplayTagsModule::RequestGameplayTag(FName(TEXT("Damage.Basic"))));
+		BaseDmgEffect->Duration.SetValue(UGameplayEffect::INSTANT_APPLICATION);
+		BaseDmgEffect->ChanceToApplyToTarget.SetValue(0.f);
+
+		// Apply to target
+		SourceComponent->ApplyGameplayEffectToTarget(BaseDmgEffect, DestComponent, 0.f);
+
+		float ExpectedValue = (StartHealth - DamageValue);
+		float ActualValue = DestComponent->GetSet<USkillSystemTestAttributeSet>()->Health;
+
+		Test->TestTrue(SKILL_TEST_TEXT("Damage Applied. %.2f == %.2f", ActualValue, ExpectedValue), ActualValue == ExpectedValue);
+	}
+
+	World->EditorDestroyActor(SourceActor, false);
+	World->EditorDestroyActor(DestActor, false);
+
+	return true;
+}
+
+bool GameplayEffectsTest_ModifyChanceToExecuteOnGE(UWorld *World, FAutomationTestBase * Test)
+{
+	const float StartHealth = 100.f;
+	const float DamageValue = 5.f;
+
+	ASkillSystemTestPawn *SourceActor = World->SpawnActor<ASkillSystemTestPawn>();
+	ASkillSystemTestPawn *DestActor = World->SpawnActor<ASkillSystemTestPawn>();
+
+	UProperty *DamageProperty = FindFieldChecked<UProperty>(USkillSystemTestAttributeSet::StaticClass(), GET_MEMBER_NAME_CHECKED(USkillSystemTestAttributeSet, Damage));
+
+	UAttributeComponent * SourceComponent = Cast<UAttributeComponent>(SourceActor->CreateComponentFromTemplate(UAttributeComponent::StaticClass()->GetDefaultObject<UAttributeComponent>(), FString(TEXT("AttributeComponent"))));
+	UAttributeComponent * DestComponent = Cast<UAttributeComponent>(DestActor->CreateComponentFromTemplate(UAttributeComponent::StaticClass()->GetDefaultObject<UAttributeComponent>(), FString(TEXT("AttributeComponent"))));
+	SourceComponent->GetSet<USkillSystemTestAttributeSet>()->Health = StartHealth;
+	DestComponent->GetSet<USkillSystemTestAttributeSet>()->Health = StartHealth;
+
+	{
+		SKILL_LOG_SCOPE(TEXT("Apply Damage buff"))
+
+		UGameplayEffect * BuffEffect = Cast<UGameplayEffect>(StaticConstructObject(UGameplayEffect::StaticClass(), GetTransientPackage(), FName(TEXT("DamageBuff"))));
+		BuffEffect->Modifiers.SetNum(1);
+		BuffEffect->Modifiers[0].Magnitude.SetValue(2.f);
+		BuffEffect->Modifiers[0].ModifierType = EGameplayMod::OutgoingGE;
+		BuffEffect->Modifiers[0].ModifierOp = EGameplayModOp::Multiplicitive;
+		BuffEffect->Modifiers[0].EffectType = EGameplayModEffect::Magnitude;
+		BuffEffect->Modifiers[0].Attribute.SetUProperty(DamageProperty);
+		BuffEffect->Modifiers[0].RequiredTags.AddTag(IGameplayTagsModule::RequestGameplayTag(FName(TEXT("Damage.Basic"))));
+		BuffEffect->Duration.SetValue(UGameplayEffect::INFINITE_DURATION);
+		BuffEffect->GameplayEffectTags.AddTag(IGameplayTagsModule::RequestGameplayTag(FName(TEXT("Damage.Buff"))));
+		BuffEffect->ChanceToExecuteOnGameplayEffect.SetValue(0.f);
+
+		// Apply to self
+		SourceComponent->ApplyGameplayEffectToTarget(BuffEffect, SourceComponent, 0.f);
+	}
+
+	// verify that outgoing damage is unbuffed
+	{
+		SKILL_LOG_SCOPE(TEXT("Apply Damage mod"))
+
+		UGameplayEffect * BaseDmgEffect = Cast<UGameplayEffect>(StaticConstructObject(UGameplayEffect::StaticClass(), GetTransientPackage(), FName(TEXT("Damage"))));
+		BaseDmgEffect->Modifiers.SetNum(1);
+		BaseDmgEffect->Modifiers[0].Magnitude.SetValue(DamageValue);
+		BaseDmgEffect->Modifiers[0].ModifierType = EGameplayMod::Attribute;
+		BaseDmgEffect->Modifiers[0].ModifierOp = EGameplayModOp::Additive;
+		BaseDmgEffect->Modifiers[0].Attribute.SetUProperty(DamageProperty);
+		BaseDmgEffect->Modifiers[0].OwnedTags.AddTag(IGameplayTagsModule::RequestGameplayTag(FName(TEXT("Damage.Basic"))));
+		BaseDmgEffect->Duration.SetValue(UGameplayEffect::INSTANT_APPLICATION);
+
+		// Apply to target
+		SourceComponent->ApplyGameplayEffectToTarget(BaseDmgEffect, DestComponent, 0.f);
+
+		float ExpectedValue = (StartHealth - DamageValue);
+		float ActualValue = DestComponent->GetSet<USkillSystemTestAttributeSet>()->Health;
+
+		Test->TestTrue(SKILL_TEST_TEXT("Damage Applied. %.2f == %.2f", ActualValue, ExpectedValue), ActualValue == ExpectedValue);
+	}
+
+	{
+		SKILL_LOG_SCOPE(TEXT("Apply boost to chance to execute"))
+
+		UGameplayEffect * BaseEffect = Cast<UGameplayEffect>(StaticConstructObject(UGameplayEffect::StaticClass(), GetTransientPackage(), FName(TEXT("ChanceToExecute"))));
+		BaseEffect->Modifiers.SetNum(1);
+		BaseEffect->Modifiers[0].Magnitude.SetValue(1.f);
+		BaseEffect->Modifiers[0].ModifierType = EGameplayMod::ActiveGE;
+		BaseEffect->Modifiers[0].ModifierOp = EGameplayModOp::Additive;
+		BaseEffect->Modifiers[0].EffectType = EGameplayModEffect::ChanceExecuteEffect;
+		BaseEffect->Duration.SetValue(UGameplayEffect::INFINITE_DURATION);
+
+		// Apply to self
+		SourceComponent->ApplyGameplayEffectToTarget(BaseEffect, SourceComponent, 0.f);
+	}
+
+	// verify that the buff applies now
+	{
+		SKILL_LOG_SCOPE(TEXT("Apply Damage mod"))
+
+		UGameplayEffect * BaseDmgEffect = Cast<UGameplayEffect>(StaticConstructObject(UGameplayEffect::StaticClass(), GetTransientPackage(), FName(TEXT("Damage"))));
+		BaseDmgEffect->Modifiers.SetNum(1);
+		BaseDmgEffect->Modifiers[0].Magnitude.SetValue(DamageValue);
+		BaseDmgEffect->Modifiers[0].ModifierType = EGameplayMod::Attribute;
+		BaseDmgEffect->Modifiers[0].ModifierOp = EGameplayModOp::Additive;
+		BaseDmgEffect->Modifiers[0].Attribute.SetUProperty(DamageProperty);
+		BaseDmgEffect->Modifiers[0].OwnedTags.AddTag(IGameplayTagsModule::RequestGameplayTag(FName(TEXT("Damage.Basic"))));
+		BaseDmgEffect->Duration.SetValue(UGameplayEffect::INSTANT_APPLICATION);
+
+		// Apply to target
+		SourceComponent->ApplyGameplayEffectToTarget(BaseDmgEffect, DestComponent, 0.f);
+
+		float ExpectedValue = (StartHealth - DamageValue * 3.f);
+		float ActualValue = DestComponent->GetSet<USkillSystemTestAttributeSet>()->Health;
+
+		Test->TestTrue(SKILL_TEST_TEXT("Damage Applied. %.2f == %.2f", ActualValue, ExpectedValue), ActualValue == ExpectedValue);
+	}
+
+	World->EditorDestroyActor(SourceActor, false);
+	World->EditorDestroyActor(DestActor, false);
+
+	return true;
+}
+
+
+
 bool FGameplayEffectsTest::RunTest( const FString& Parameters )
 {
 	UCurveTable *CurveTable = ISkillSystemModule::Get().GetSkillSystemGlobals().GetGlobalCurveTable();
@@ -4406,6 +5777,36 @@ bool FGameplayEffectsTest::RunTest( const FString& Parameters )
 	GameplayEffectsTest_StackingRemovingModifiers(World, this);
 	GameplayEffectsTest_StackingAddingModifiers(World, this);
 
+	// Immunity
+	GameplayEffectsTest_ImmunityIncoming(World, this);
+	GameplayEffectsTest_ImmunityOutgoing(World, this);
+	GameplayEffectsTest_ImmunityActiveGE(World, this);
+	GameplayEffectsTest_ImmunityMod(World, this);
+	GameplayEffectsTest_ImmunityActiveMod(World, this);
+	GameplayEffectsTest_ImmunityBlockedBuff(World, this);
+	GameplayEffectsTest_ImmunityBlockedBaseAndBuff(World, this);
+
+	//
+	// Chance to apply or execute
+	//
+
+	// Chance to apply to target
+	GameplayEffectsTest_ChanceToApplyToTarget(World, this);
+
+	// Chance to apply to GEs
+	// We need to test active, incoming and outgoing GEs
+	// Chance to execute has a slightly different path for immunity than it does for modifying a GE so it needs to be tested separately
+	GameplayEffectsTest_ChanceToExecuteOnActiveGEMod(World, this);
+	GameplayEffectsTest_ChanceToExecuteOnActiveGEImmunity(World, this);
+	GameplayEffectsTest_ChanceToExecuteOnOutgoingGEMod(World, this);
+	GameplayEffectsTest_ChanceToExecuteOnOutgoingGEImmunity(World, this);
+	GameplayEffectsTest_ChanceToExecuteOnIncomingGEMod(World, this);
+	GameplayEffectsTest_ChanceToExecuteOnIncomingGEImmunity(World, this);
+
+	// Modifiers to Chance to apply and execute
+	GameplayEffectsTest_ModifyChanceToApplyToTarget(World, this);
+	GameplayEffectsTest_ModifyChanceToExecuteOnGE(World, this);
+
 	GEngine->DestroyWorldContext(World);
 	World->DestroyWorld(false);
 
