Code cleanup.

Removed commented out blocks of code.

diff --git a/Engine/Source/Runtime/SkillSystem/Private/GameplayEffectStackingExtension_DiminishingReturnsTest.cpp b/Engine/Source/Runtime/SkillSystem/Private/GameplayEffectStackingExtension_DiminishingReturnsTest.cpp
--- a/Engine/Source/Runtime/SkillSystem/Private/GameplayEffectStackingExtension_DiminishingReturnsTest.cpp
+++ b/Engine/Source/Runtime/SkillSystem/Private/GameplayEffectStackingExtension_DiminishingReturnsTest.cpp
@@ -33,36 +33,6 @@ void UGameplayEffectStackingExtension_DiminishingReturnsTest::CalculateStack(TAr
 		}
 	}
 
-// 	// remove any old stack modifiers
-// 	for (int32 Idx = 0; Idx < CurrentEffect.Spec.Modifiers.Num(); ++Idx)
-// 	{
-// 		if (CurrentEffect.Spec.Modifiers[Idx].Info.OwnedTags.HasTag("Stack"))
-// 		{
-// 			CurrentEffect.Spec.Modifiers.RemoveAtSwap(Idx);
-// 			--Idx;
-// 		}
-// 	}
-// 	
-// 	for (FModifierSpec Mod : CurrentEffect.Spec.Modifiers)
-// 	{
-// 		if (Mod.Info.OwnedTags.HasTag("Stackable"))
-// 		{
-// 			FGameplayModifierInfo ModInfo;
-// 			ModInfo.Magnitude.SetValue(Multiplier);
-// 			ModInfo.ModifierOp = EGameplayModOp::Multiplicitive;
-// 			ModInfo.OwnedTags.AddTag("Stack");
-// 			ModInfo.Attribute = Mod.Info.Attribute;
-// 
-// 			TSharedPtr<FGameplayEffectLevelSpec> ModifierLevel(TSharedPtr< FGameplayEffectLevelSpec >(new FGameplayEffectLevelSpec()));
-// 			ModifierLevel->ApplyNewDef(ModInfo.LevelInfo, ModifierLevel);
-// 
-// 			FModifierSpec ModSpec(ModInfo, ModifierLevel, NULL);
-// 
-// 			ModSpec.ApplyModTo(Mod, false);
-// 			CurrentEffect.Spec.Modifiers.Add(ModSpec);
-// 		}
-// 	}
-
 	for (FModifierSpec Mod : CurrentEffect.Spec.Modifiers)
 	{
 		if (Mod.Info.OwnedTags.HasTag(IGameplayTagsModule::RequestGameplayTag("Stackable"), EGameplayTagMatchType::IncludeParentTags, EGameplayTagMatchType::Explicit))
