Cast enums to int32s before comparing values.

Some compilers get upset when comparing enums of different types for equality.

diff --git a/Engine/Source/Runtime/SkillSystem/Classes/GameplayEffect.h b/Engine/Source/Runtime/SkillSystem/Classes/GameplayEffect.h
--- a/Engine/Source/Runtime/SkillSystem/Classes/GameplayEffect.h
+++ b/Engine/Source/Runtime/SkillSystem/Classes/GameplayEffect.h
@@ -79,10 +79,10 @@ namespace EGameplayImmunity
 	};
 }
 
-checkAtCompileTime(EGameplayMod::OutgoingGE == EGameplayImmunity::OutgoingGE, "EGameplayMod::OutgoingGE and EGameplayImmunity::OutgoingGE must match. Did you forget to modify one of them?");
-checkAtCompileTime(EGameplayMod::IncomingGE == EGameplayImmunity::IncomingGE, "EGameplayMod::IncomingGE and EGameplayImmunity::IncomingGE must match. Did you forget to modify one of them?");
-checkAtCompileTime(EGameplayMod::ActiveGE == EGameplayImmunity::ActiveGE, "EGameplayMod::ActiveGE and EGameplayImmunity::ActiveGE must match. Did you forget to modify one of them?");
-checkAtCompileTime(EGameplayMod::Max == EGameplayImmunity::Max, "EGameplayMod and EGameplayImmunity must cover the same cases. Did you forget to modify one of them?");
+checkAtCompileTime((int32)EGameplayMod::OutgoingGE == (int32)EGameplayImmunity::OutgoingGE, "EGameplayMod::OutgoingGE and EGameplayImmunity::OutgoingGE must match. Did you forget to modify one of them?");
+checkAtCompileTime((int32)EGameplayMod::IncomingGE == (int32)EGameplayImmunity::IncomingGE, "EGameplayMod::IncomingGE and EGameplayImmunity::IncomingGE must match. Did you forget to modify one of them?");
+checkAtCompileTime((int32)EGameplayMod::ActiveGE == (int32)EGameplayImmunity::ActiveGE, "EGameplayMod::ActiveGE and EGameplayImmunity::ActiveGE must match. Did you forget to modify one of them?");
+checkAtCompileTime((int32)EGameplayMod::Max == (int32)EGameplayImmunity::Max, "EGameplayMod and EGameplayImmunity must cover the same cases. Did you forget to modify one of them?");
 
 /**
  * Tells us what a GameplayEffect modifies when being applied to another GameplayEffect
diff --git a/Engine/Source/Runtime/SkillSystem/Private/GameplayEffect.cpp b/Engine/Source/Runtime/SkillSystem/Private/GameplayEffect.cpp
--- a/Engine/Source/Runtime/SkillSystem/Private/GameplayEffect.cpp
+++ b/Engine/Source/Runtime/SkillSystem/Private/GameplayEffect.cpp
@@ -191,7 +191,7 @@ int32 FGameplayEffectSpec::ApplyModifiersFrom(FGameplayEffectSpec &InSpec, const
 	if (!InSpec.Def || !InSpec.Def->AreGameplayEffectTagRequirementsSatisfied(Def))
 	{
 		// InSpec doesn't match this gameplay effect but if InSpec provides immunity we also need to check the modifiers because they can be canceled independent of the gameplay effect
-		if (InSpec.Def->AppliesImmunityTo == QualifierContext.Type())
+		if ((int32)InSpec.Def->AppliesImmunityTo == (int32)QualifierContext.Type())
 		{
 			for (int ii = 0; ii < Modifiers.Num(); ++ii)
 			{
@@ -207,7 +207,7 @@ int32 FGameplayEffectSpec::ApplyModifiersFrom(FGameplayEffectSpec &InSpec, const
 		return 0;
 	}
 
-	if (InSpec.Def->AppliesImmunityTo == QualifierContext.Type())
+	if ((int32)InSpec.Def->AppliesImmunityTo == (int32)QualifierContext.Type())
 	{
 		return -1;
 	}
