Skillsystem updates
-Mostly abilities work

diff --git a/Engine/Source/Editor/SkillSystemEditor/Classes/K2Node_WaitMovementModeChange.h b/Engine/Source/Editor/SkillSystemEditor/Classes/K2Node_WaitMovementModeChange.h
new file mode 100644
--- /dev/null
+++ b/Engine/Source/Editor/SkillSystemEditor/Classes/K2Node_WaitMovementModeChange.h
@@ -0,0 +1,17 @@
+// Copyright 1998-2014 Epic Games, Inc. All Rights Reserved.
+
+#pragma once
+#include "K2Node_WaitMovementModeChange.generated.h"
+
+UCLASS()
+class UK2Node_WaitMovementModeChange : public UK2Node_BaseAsyncTask
+{
+	GENERATED_UCLASS_BODY()
+
+	// Begin UEdGraphNode interface
+	virtual FString GetTooltip() const OVERRIDE;
+	virtual FText GetNodeTitle(ENodeTitleType::Type TitleType) const OVERRIDE;
+	// End UEdGraphNode interface
+
+	virtual FString GetCategoryName();
+};
diff --git a/Engine/Source/Editor/SkillSystemEditor/Private/K2Node_WaitMovementModeChange.cpp b/Engine/Source/Editor/SkillSystemEditor/Private/K2Node_WaitMovementModeChange.cpp
new file mode 100644
--- /dev/null
+++ b/Engine/Source/Editor/SkillSystemEditor/Private/K2Node_WaitMovementModeChange.cpp
@@ -0,0 +1,37 @@
+// Copyright 1998-2014 Epic Games, Inc. All Rights Reserved.
+
+
+#include "SkillSystemEditorModulePrivatePCH.h"
+#include "SkillSystemBlueprintLibrary.h"
+#include "BlueprintWaitMovementModeChangeTaskProxy.h"
+#include "CompilerResultsLog.h"
+
+#define LOCTEXT_NAMESPACE "K2Node_WaitMovementModeChange"
+
+UK2Node_WaitMovementModeChange::UK2Node_WaitMovementModeChange(const class FPostConstructInitializeProperties& PCIP)
+	: Super(PCIP)
+{
+	
+	ProxyFactoryFunctionName = GET_FUNCTION_NAME_CHECKED(USkillSystemBlueprintLibrary, CreateWaitMovementModeChange);
+	ProxyFactoryClass = USkillSystemBlueprintLibrary::StaticClass();
+	ProxyClass = UBlueprintWaitMovementModeChangeTaskProxy::StaticClass();
+}
+
+FString UK2Node_WaitMovementModeChange::GetCategoryName()
+{
+	return TEXT("Ability");
+}
+
+FString UK2Node_WaitMovementModeChange::GetTooltip() const
+{
+	return TEXT("Waits until the movement mode changes");
+}
+
+FText UK2Node_WaitMovementModeChange::GetNodeTitle(ENodeTitleType::Type TitleType) const
+{
+	return LOCTEXT("WaitMovementModeChange", "Wait Movement Mode Change");
+}
+
+#undef LOCTEXT_NAMESPACE
+
+
diff --git a/Engine/Source/Runtime/SkillSystem/Classes/Abilities/GameplayAbility.h b/Engine/Source/Runtime/SkillSystem/Classes/Abilities/GameplayAbility.h
new file mode 100644
--- /dev/null
+++ b/Engine/Source/Runtime/SkillSystem/Classes/Abilities/GameplayAbility.h
@@ -0,0 +1,312 @@
+// Copyright 1998-2013 Epic Games, Inc. All Rights Reserved.
+
+#pragma once
+
+#include "GameplayTagContainer.h"
+#include "GameplayTagAssetInterface.h"
+#include "Runtime/Engine/Classes/Animation/AnimInstance.h"
+#include "SkillSystemTypes.h"
+#include "GameplayAbility.generated.h"
+
+class UAnimInstance;
+class UAttributeComponent;
+
+UENUM(BlueprintType)
+namespace EGameplayAbilityInstancingPolicy
+{
+	/**
+	 *	How the ability is instanced when executed. This limits what an ability can do in its implementation. For example, a NonInstanced
+	 *	Ability cannot have state. It is probably unsafe for an InstancedPerActor ability to have latent actions, etc.	
+	 */
+
+	enum Type
+	{
+		NonInstanced,			// This ability is never instanced. Anything that executes the ability is operating on the CDO.
+		InstancedPerActor,		// Each actor gets their own instance of this ability. State can be saved, replication is possible.
+		InstancedPerExecution,	// We instance this ability each time it is executed. Replication possible but not recommended.
+	};
+}
+
+UENUM(BlueprintType)
+namespace EGameplayAbilityNetExecutionPolicy
+{
+	/**
+	 *	How does an ability execute on the network. Does a client "ask and predict", "ask and wait", "don't ask" 
+	 */
+
+	enum Type
+	{
+		Predictive		UMETA(DisplayName = "Predictive"),	// Part of this ability runs predictively on the client.
+		Server			UMETA(DisplayName = "Server"),		// This ability must be OK'd by the server before doing anything on a client.
+		Client			UMETA(DisplayName = "Client"),		// This ability runs as long the client says it does.
+	};
+}
+
+UENUM(BlueprintType)
+namespace EGameplayAbilityReplicationPolicy
+{
+	/**
+	 *	How an ability replicates state/events to everyone on the network.
+	 */
+
+	enum Type
+	{
+		ReplicateNone		UMETA(DisplayName = "Replicate None"),	// We don't replicate the instance of the ability to anyone.
+		ReplicateAll		UMETA(DisplayName = "Replicate All"),	// We replicate the instance of the ability to everyone (even simulating clients).
+		ReplicateOwner		UMETA(DisplayName = "Replicate Owner"),	// Only replicate the instance of the ability to the owner.
+	};
+}
+
+USTRUCT(BlueprintType)
+struct FGameplayAbilityHandle
+{
+	GENERATED_USTRUCT_BODY()
+
+	FGameplayAbilityHandle()
+	: Handle(INDEX_NONE)
+	{
+
+	}
+
+	FGameplayAbilityHandle(int32 InHandle)
+		: Handle(InHandle)
+	{
+
+	}
+
+	bool IsValid() const
+	{
+		return Handle != INDEX_NONE;
+	}
+
+	FGameplayAbilityHandle GetNextHandle()
+	{
+		return FGameplayAbilityHandle(Handle + 1);
+	}
+
+	bool operator==(const FGameplayAbilityHandle& Other) const
+	{
+		return Handle == Other.Handle;
+	}
+
+	bool operator!=(const FGameplayAbilityHandle& Other) const
+	{
+		return Handle != Other.Handle;
+	}
+
+	FString ToString() const
+	{
+		return FString::Printf(TEXT("%d"), Handle);
+	}
+
+private:
+
+	UPROPERTY()
+	int32 Handle;
+};
+
+USTRUCT()
+struct FGameplayAbilitySpec
+{
+	GENERATED_USTRUCT_BODY()
+
+	UPROPERTY()
+	int32	Level;
+
+	UPROPERTY()
+	class UGameplayAbility * Ability;
+
+	UPROPERTY()
+	FGameplayAbilityHandle	Handle;
+};
+
+/**
+* UGameplayAbility
+*	The GameplayAbility baseclass. This base class is a data asset and can be used to create non-instanced abilities.
+*	Instanced abilities should derive off UGameplyAbility_Instanced.
+*/
+UCLASS()
+class SKILLSYSTEM_API UGameplayAbility : public UDataAsset
+{
+	GENERATED_UCLASS_BODY()
+
+public:
+
+	UFUNCTION(Client, Reliable)
+	void ClientActivateAbilitySucceed(AActor *Actor);
+
+	bool IsSupportedForNetworking() const OVERRIDE;
+	
+	// ----------------------------------------------------------------------------------------------------------------
+	//
+	//	The important functions:
+	//	
+	//		CanActivateAbility()	- const function to see if ability is activatable. Callable by UI etc
+	//
+	//		TryActivateAbility()	- Attempts to activate the ability. Calls CanActivateAbility(). Input events can call this directly.
+	//								- Also handles instancing-per-execution logic and replication/prediction calls.
+	//		
+	//		CallActivate()			- Protected, non virtual function. Does some boilerplate 'pre activate' stuff, then calls Activate()
+	//
+	//		Activate()				- What the abilities *does*. This is what child classes want to override.
+	//	
+	//		Commit()				- Commits reources/cooldowns etc. Activate() must call this!
+	//		
+	//		CancelAbility()			- 
+	//	
+	// ----------------------------------------------------------------------------------------------------------------
+
+	/**
+	 * Attempts to activate the ability.
+	 *	-This function calls CanActivateAbility
+	 *	-This function handles instancing
+	 *	-This function handles networking and prediction
+	 *	-If all goes well, CallActivateAbility is called next.
+	 */
+	bool TryActivateAbility(const FGameplayAbilityActorInfo ActorInfo, UGameplayAbility ** OutInstancedAbility = NULL);
+	
+	/** Returns true if this ability can be activated right now. Has no side effects! */
+	virtual bool CanActivateAbility(const FGameplayAbilityActorInfo ActorInfo) const;		
+
+	// -----------------------------------------------
+
+	/** Input binding. Base implementation calls TryActivateAbility */
+	virtual void InputPressed(int32 InputID, const FGameplayAbilityActorInfo ActorInfo);
+
+	/** Input binding. Base implementation does nothing */
+	virtual void InputReleased(int32 InputID, const FGameplayAbilityActorInfo ActorInfo);
+
+	// ----------------------------------------------------------------------------------------------------------------
+	//
+	//	Outward facing API for other systems
+	//
+	// ----------------------------------------------------------------------------------------------------------------
+
+	float GetCooldownTimeRemaining(const FGameplayAbilityActorInfo ActorInfo) const;
+
+	// ----------------------------------------------------------------------------------------------------------------
+	//
+	//	Standardized State
+	//
+	// ----------------------------------------------------------------------------------------------------------------	
+
+	UPROPERTY(EditDefaultsOnly, Category=Advanced)
+	TEnumAsByte<EGameplayAbilityNetExecutionPolicy::Type>	NetExecutionPolicy;
+
+	/** This GameplayEffect represents the cooldown. It will be applied when the ability is commited and the ability cannot be used again until it is expired. */
+	UPROPERTY(EditDefaultsOnly, Category=Cooldowns)
+	class UGameplayEffect * CooldownGameplayEffect;
+
+	/** This GameplayEffect represents the cost (mana, stamina, etc) of the ability. It will be applied when the ability is commited. */
+	UPROPERTY(EditDefaultsOnly, Category=Costs)
+	class UGameplayEffect * CostGameplayEffect;
+
+	// ----------------------------------------------------------------------------------------------------------------
+	//
+	//	Ability exclusion / cancelling
+	//
+	// ----------------------------------------------------------------------------------------------------------------
+	
+	/** Abilities with these tags are canelled when this ability is executed */
+	UPROPERTY(EditDefaultsOnly, Category = Tags)
+	FGameplayTagContainer CancelAbilitiesWithTag;
+
+	/** This ability has these tags */
+	UPROPERTY(EditDefaultsOnly, Category = Tags)
+	FGameplayTagContainer AbilityTags;
+
+	// ----------------------------------------------------------------------------------------------------------------
+	//
+	//	Animation callbacks
+	//
+	// ----------------------------------------------------------------------------------------------------------------
+
+	void MontageBranchPoint_AbilityDecisionStop(const FGameplayAbilityActorInfo ActorInfo) const;
+
+	void MontageBranchPoint_AbilityDecisionStart(const FGameplayAbilityActorInfo ActorInfo) const;
+
+	// ----------------------------------------------------------------------------------------------------------------
+
+	virtual void PredictiveActivateAbility(const FGameplayAbilityActorInfo ActorInfo);
+
+	virtual void ServerTryActivateAbility(const FGameplayAbilityActorInfo ActorInfo);
+
+	virtual void ClientActivateAbility(const FGameplayAbilityActorInfo ActorInfo);
+
+	// ----------------------------------------------------------------------------------------------------------------
+
+	virtual EGameplayAbilityInstancingPolicy::Type GetInstancingPolicy() const
+	{
+		return EGameplayAbilityInstancingPolicy::NonInstanced;
+	}
+
+	virtual EGameplayAbilityReplicationPolicy::Type GetReplicationPolicy() const
+	{
+		return EGameplayAbilityReplicationPolicy::ReplicateNone;
+	}
+
+	// ----------------------------------------------------------------------------------------------------------------
+
+	virtual UWorld* GetWorld() const OVERRIDE
+	{
+		return GetOuter()->GetWorld();
+	}
+
+protected:
+
+	/**
+	 * The main function that defines what an ability does.
+	 *  -Child classes will want to override this
+	 *  -This function must call CommitAbility()
+	 */
+
+	// FIXME: Making this blueprintcallable is danagerous!
+	UFUNCTION(BlueprintCallable, Category = "Ability")
+	virtual void ActivateAbility(const FGameplayAbilityActorInfo ActorInfo);
+
+	/** 
+	 * Attempts to commit the ability (spend resources, etc). This our last chance to fail.
+	 *	-Child classes that override ActivateAbility must call this themselves!
+	 */
+	UFUNCTION(BlueprintCallable, Category = "Ability")
+	virtual bool CommitAbility(const FGameplayAbilityActorInfo ActorInfo);
+	
+	/** Do boilerplate init stuff and then call ActivateAbility */
+	void PreActivate(const FGameplayAbilityActorInfo ActorInfo);
+
+	void CallActivateAbility(const FGameplayAbilityActorInfo ActorInfo);
+
+	void CallPredictiveActivateAbility(const FGameplayAbilityActorInfo ActorInfo);
+
+	UPROPERTY(EditDefaultsOnly, Category = Tags)
+	bool Prediction;
+
+	/** 
+	 * The last chance to fail before commiting
+	 *	-This will usually be the same as CanActivateAbility. Some abilities may need to do extra checks here if they are consuming extra stuff in CommitExecute
+	 */
+	virtual bool CommitCheck(const FGameplayAbilityActorInfo ActorInfo);
+
+	/** Does the commmit atomically (consume resources, do cooldowns, etc) */
+	virtual void CommitExecute(const FGameplayAbilityActorInfo ActorInfo);
+
+	/** Destroys intanced-per-execution abilities. Instance-per-actor abilities should 'reset'. Non instance abilities - what can we do? */
+	virtual void CancelAbility(const FGameplayAbilityActorInfo ActorInfo) { }
+
+	// --------------------------------------------------------------------------------
+
+
+	// --------------------------------------------------------------------------------
+
+	/** Checks cooldown. returns true if we can be used again. False if not */
+	bool	CheckCooldown(const FGameplayAbilityActorInfo ActorInfo) const;
+
+	/** Applies CooldownGameplayEffect to the target */
+	void	ApplyCooldown(const FGameplayAbilityActorInfo ActorInfo);
+
+	bool	CheckCost(const FGameplayAbilityActorInfo ActorInfo) const;
+
+	void	ApplyCost(const FGameplayAbilityActorInfo ActorInfo);
+
+	friend class UAttributeComponent;
+};
\ No newline at end of file
diff --git a/Engine/Source/Runtime/SkillSystem/Classes/Abilities/GameplayAbility_CharacterJump.h b/Engine/Source/Runtime/SkillSystem/Classes/Abilities/GameplayAbility_CharacterJump.h
new file mode 100644
--- /dev/null
+++ b/Engine/Source/Runtime/SkillSystem/Classes/Abilities/GameplayAbility_CharacterJump.h
@@ -0,0 +1,26 @@
+// Copyright 1998-2013 Epic Games, Inc. All Rights Reserved.
+
+#pragma once
+
+#include "GameplayAbility_CharacterJump.generated.h"
+
+/**
+ *	A gameplay ability that plays a single montage and applies a GameplayEffect
+ */
+UCLASS()
+class SKILLSYSTEM_API UGameplayAbility_CharacterJump : public UGameplayAbility
+{
+	GENERATED_UCLASS_BODY()
+
+public:
+
+	virtual bool CanActivateAbility(const FGameplayAbilityActorInfo ActorInfo) const OVERRIDE;
+	
+	virtual void ActivateAbility(const FGameplayAbilityActorInfo OwnerInfo) OVERRIDE;
+
+	virtual void InputReleased(int32 InputID, const FGameplayAbilityActorInfo ActorInfo) OVERRIDE;
+
+	virtual void PredictiveActivateAbility(const FGameplayAbilityActorInfo ActorInfo) OVERRIDE;
+
+	virtual void CancelAbility(const FGameplayAbilityActorInfo ActorInfo) OVERRIDE;
+};
\ No newline at end of file
diff --git a/Engine/Source/Runtime/SkillSystem/Classes/Abilities/GameplayAbility_Instanced.h b/Engine/Source/Runtime/SkillSystem/Classes/Abilities/GameplayAbility_Instanced.h
new file mode 100644
--- /dev/null
+++ b/Engine/Source/Runtime/SkillSystem/Classes/Abilities/GameplayAbility_Instanced.h
@@ -0,0 +1,73 @@
+// Copyright 1998-2013 Epic Games, Inc. All Rights Reserved.
+
+#pragma once
+
+#include "GameplayAbility_Instanced.generated.h"
+
+class UAnimInstance;
+class UAttributeComponent;
+
+/**
+* UGameplayEffect
+*	The GameplayEffect definition. This is the data asset defined in the editor that drives everything.
+*/
+UCLASS(Blueprintable)
+class SKILLSYSTEM_API UGameplayAbility_Instanced : public UGameplayAbility
+{
+	GENERATED_UCLASS_BODY()
+
+public:
+	
+	// -------------------------------------------------
+
+	UFUNCTION(BlueprintImplementableEvent, Category = Ability)
+	virtual bool K2_CanActivateAbility(const FGameplayAbilityActorInfo ActorInfo) const;
+
+	virtual bool CanActivateAbility(const FGameplayAbilityActorInfo ActorInfo) const OVERRIDE;
+
+	// -------------------------------------------------
+
+	UFUNCTION(BlueprintImplementableEvent, Category = Ability)
+	virtual void K2_CommitExecute(const FGameplayAbilityActorInfo ActorInfo);
+
+	virtual void CommitExecute(const FGameplayAbilityActorInfo ActorInfo) OVERRIDE;
+
+	// -------------------------------------------------
+
+	UFUNCTION(BlueprintImplementableEvent, Category = Ability)
+	virtual void K2_ActivateAbility(const FGameplayAbilityActorInfo ActorInfo);
+
+	virtual void ActivateAbility(const FGameplayAbilityActorInfo ActorInfo) OVERRIDE;
+
+	// -------------------------------------------------
+
+	UFUNCTION(BlueprintImplementableEvent, Category = Ability)
+	virtual void K2_PredictiveActivateAbility(const FGameplayAbilityActorInfo ActorInfo);
+
+	virtual void PredictiveActivateAbility(const FGameplayAbilityActorInfo ActorInfo) OVERRIDE;
+
+	// -------------------------------------------------
+
+	virtual EGameplayAbilityInstancingPolicy::Type GetInstancingPolicy() const OVERRIDE
+	{
+		return InstancedPerExecution ? EGameplayAbilityInstancingPolicy::InstancedPerExecution : EGameplayAbilityInstancingPolicy::InstancedPerActor;
+	}
+
+	virtual EGameplayAbilityReplicationPolicy::Type GetReplicationPolicy() const OVERRIDE
+	{
+		return ReplicationPolicy;
+	}
+
+	UPROPERTY(EditDefaultsOnly, Category=Advanced)
+	TEnumAsByte<EGameplayAbilityReplicationPolicy::Type> ReplicationPolicy;
+
+	UPROPERTY(EditDefaultsOnly, Category=Advanced)
+	bool	InstancedPerExecution;
+
+	int32 GetFunctionCallspace(UFunction* Function, void* Parameters, FFrame* Stack);
+
+	bool CallRemoteFunction(UFunction* Function, void* Parameters, FOutParmRec* OutParms, FFrame* Stack);
+
+private:
+	
+};
\ No newline at end of file
diff --git a/Engine/Source/Runtime/SkillSystem/Classes/Abilities/GameplayAbility_Montage.h b/Engine/Source/Runtime/SkillSystem/Classes/Abilities/GameplayAbility_Montage.h
new file mode 100644
--- /dev/null
+++ b/Engine/Source/Runtime/SkillSystem/Classes/Abilities/GameplayAbility_Montage.h
@@ -0,0 +1,33 @@
+// Copyright 1998-2013 Epic Games, Inc. All Rights Reserved.
+
+#pragma once
+
+#include "GameplayAbility_Montage.generated.h"
+
+/**
+ *	A gameplay ability that plays a single montage and applies a GameplayEffect
+ */
+UCLASS()
+class SKILLSYSTEM_API UGameplayAbility_Montage : public UGameplayAbility
+{
+	GENERATED_UCLASS_BODY()
+
+public:
+	
+	virtual void ActivateAbility(const FGameplayAbilityActorInfo OwnerInfo) OVERRIDE;
+
+	UPROPERTY(EditDefaultsOnly, Category = MontageAbility)
+	UAnimMontage *	MontageToPlay;
+
+	UPROPERTY(EditDefaultsOnly, Category = MontageAbility)
+	float	PlayRate;
+
+	UPROPERTY(EditDefaultsOnly, Category = MontageAbility)
+	FName	SectionName;
+
+	/** GameplayEffects to apply and then remove while the animation is playing */
+	UPROPERTY(EditDefaultsOnly, Category = MontageAbility)
+	TArray<const UGameplayEffect*>	GameplayEffectsWhileAnimating;
+
+	void OnMontageEnded(UAnimMontage* Montage, bool bInterrupted, const FGameplayAbilityActorInfo OwnerInfo, TArray<struct FActiveGameplayEffectHandle>	AppliedEffects);
+};
\ No newline at end of file
diff --git a/Engine/Source/Runtime/SkillSystem/Classes/AbilityTaskProxyBase.h b/Engine/Source/Runtime/SkillSystem/Classes/AbilityTaskProxyBase.h
new file mode 100644
--- /dev/null
+++ b/Engine/Source/Runtime/SkillSystem/Classes/AbilityTaskProxyBase.h
@@ -0,0 +1,12 @@
+// Copyright 1998-2014 Epic Games, Inc. All Rights Reserved.
+#pragma once
+#include "AbilityTaskProxyBase.Generated.h"
+
+UCLASS(MinimalAPI)
+class UAbilityTaskProxyBase : public UObject
+{
+	GENERATED_UCLASS_BODY()	
+
+public:
+	
+};
\ No newline at end of file
diff --git a/Engine/Source/Runtime/SkillSystem/Classes/AttributeComponent.h b/Engine/Source/Runtime/SkillSystem/Classes/AttributeComponent.h
--- a/Engine/Source/Runtime/SkillSystem/Classes/AttributeComponent.h
+++ b/Engine/Source/Runtime/SkillSystem/Classes/AttributeComponent.h
@@ -2,10 +2,13 @@
 #pragma once
 
 #include "GameplayEffect.h"
+#include "SkillSystemTypes.h"
 #include "AttributeComponent.generated.h"
 
 SKILLSYSTEM_API DECLARE_LOG_CATEGORY_EXTERN(LogAttributeComponent, Log, All);
 
+DECLARE_MULTICAST_DELEGATE(FAttributeComponentPredictionKeyClear);
+
 USTRUCT()
 struct SKILLSYSTEM_API FAttributeDefaults
 {
@@ -27,17 +30,41 @@ class SKILLSYSTEM_API UAttributeComponent : public UActorComponent
 	GENERATED_UCLASS_BODY()
 
 	friend FGameplayEffectSpec;
+	friend class ASkillSystemDebugHUD;
 
 	virtual ~UAttributeComponent();
 
 	virtual void InitializeComponent() OVERRIDE;
 
+	/** Finds existing AttributeSet */
 	template <class T >
 	T*	GetSet()
 	{
 		return (T*)GetAttributeSubobject(T::StaticClass());
 	}
 
+	template <class T >
+	T*	GetSetChecked()
+	{
+		return (T*)GetAttributeSubobjectChecked(T::StaticClass());
+	}
+
+	/** Adds a new AttributeSet (initialized to default values) */
+	template <class T >
+	T*  AddSet()
+	{
+		return (T*)GetOrCreateAttributeSubobject(T::StaticClass());
+	}
+
+	/** Adds a new AttributeSet that is a DSO (created by called in their CStor) */
+	template <class T>
+	T*	AddDefaultSubobjectSet(TSubobjectPtr<T> Subobject)
+	{
+		T* Set = Subobject.Get();
+		SpawnedAttributes.Add(Set);
+		return Set;
+	}
+
 	UFUNCTION(BlueprintCallable, Category="Skills")
 	UAttributeSet * InitStats(TSubclassOf<class UAttributeSet> Attributes, const UDataTable* DataTable);
 
@@ -47,15 +74,72 @@ class SKILLSYSTEM_API UAttributeComponent : public UActorComponent
 	UPROPERTY(EditAnywhere, Category="AttributeTest")
 	TArray<FAttributeDefaults>	DefaultStartingData;
 
-	UPROPERTY()
+	UPROPERTY(Replicated)
 	TArray<UAttributeSet*>	SpawnedAttributes;
 
-	UPROPERTY()
-	float	TestFloat;
-
 	void SetNumericAttribute(const FGameplayAttribute &Attribute, float NewFloatValue);
 	float GetNumericAttribute(const FGameplayAttribute &Attribute);
 
+	// -- Replication -------------------------------------------------------------------------------------------------
+
+	virtual bool ReplicateSubobjects(class UActorChannel *Channel, class FOutBunch *Bunch, FReplicationFlags *RepFlags) OVERRIDE;
+	
+	virtual void GetSubobjectsWithStableNamesForNetworking(TArray<UObject*>& Objs) OVERRIDE;
+
+	virtual void PostNetReceive() OVERRIDE;
+
+	/**
+	 *	Prediction Keys
+	 *		PredictionKey is a simple system to allow clients to locally predict interactions within the AbilityComponent while having
+	 *		a mechanism for merging results from the server.
+	 *		
+	 *		There are two main cases:
+	 *			-Client predicts something and is wrong. (Mispredictions). He must cleanup stuff that he created or changed.
+	 *			-Client predicts something and is right. 
+	 *			
+	 *			The Client predicts something and asks the server. "Can I do X? PredictionKey=Y".
+	 *			Server can reply immediately and directly with yes or no, and include the PredictionKey with that answer.
+	 *			
+	 *			If No, the client can just immediately cleanup whatever he predicted. This is straightforward.
+	 *			
+	 *			If Yes, there is a small window where the Yes RPC is received but the actor/component properties have not been updated.
+	 *			The client's cannot immediately cleanup his prediction work: he needs to wait for property replication.
+	 *			
+	 *			This is where ReplicatedPredictionKey comes in. 
+	 *			
+	 *			When the server allows a predicted action, he sets ReplicatedPredictionKey to this value. The client then looks
+	 *			at this replicated value. In OnRep_PredictionKey, the client goes through and cleans up all stuff that was predicted
+	 *			and is no longer needed.
+	 *			
+	 *			-Client calls GetNextPredictionKey() at the start of a predicted action. This gives him a key to use for this 'run'.
+	 *			-Client can predict whatever he wants as long as he registers an 'undo' function with GetPredictionKeyDelegate().
+	 *			-This delegate is invoked if the client is denied his prediction, or if he is right and property replication catches up.
+	 *			
+	 *	 
+	 */
+
+	UPROPERTY(ReplicatedUsing=OnRep_PredictionKey)
+	int32	ReplicatedPredictionKey;
+
+	UPROPERTY(transient)
+	int32	LocalPredictionKey;
+
+	int32	GetNextPredictionKey();
+
+	UFUNCTION()
+	void OnRep_PredictionKey();
+
+	TArray<TPair<int32, FAttributeComponentPredictionKeyClear> >	PredictionDelegates;
+	
+
+	FAttributeComponentPredictionKeyClear &	GetPredictionKeyDelegate(int32 PredictionKey);
+
+	/*
+	virtual void OnSubobjectCreatedFromReplication(UObject *NewSubobject) OVERRIDE;
+	
+	virtual void OnSubobjectDestroyFromReplication(UObject *NewSubobject) OVERRIDE;
+	*/
+
 	// ----------------------------------------------------------------------------------------------------------------
 	//
 	//	GameplayEffects
@@ -71,6 +155,12 @@ class SKILLSYSTEM_API UAttributeComponent : public UActorComponent
 	UFUNCTION(BlueprintCallable, BlueprintAuthorityOnly, Category = GameplayEffects)
 	bool RemoveActiveGameplayEffect(FActiveGameplayEffectHandle Handle);
 
+	/** This only exists so it can be hooked up to a multicast delegate */
+	void RemoveActiveGameplayEffect_NoReturn(FActiveGameplayEffectHandle Handle)
+	{
+		RemoveActiveGameplayEffect(Handle);
+	}
+
 	float GetGameplayEffectDuration(FActiveGameplayEffectHandle Handle) const;
 
 	float GetGameplayEffectDuration() const;
@@ -89,6 +179,9 @@ class SKILLSYSTEM_API UAttributeComponent : public UActorComponent
 	// Tags
 	UFUNCTION(BlueprintCallable, Category=GameplayEffects)
 	bool HasAnyTags(FGameplayTagContainer &Tags);
+
+	UFUNCTION(BlueprintCallable, Category = GameplayEffects)
+	bool HasAllTags(FGameplayTagContainer &Tags);
 	
 
 	// --------------------------------------------
@@ -138,36 +231,64 @@ class SKILLSYSTEM_API UAttributeComponent : public UActorComponent
 		}
 	}
 
-	// ----------------------------------------------------------------------------------------------------------------
-	//
-	//	GameplayAbilities
-	//	
-	//	Should this be its own component? What should this do?
-	//		-Load/initialize instanced subobjects
-	//		-Bind to input component (maybe? where would weapons fit in?)
-	//		
-	//
 	// ----------------------------------------------------------------------------------------------------------------
 
-	/** Creates instances of the abilities that we need but doesn't necessarily  */
-	void InitializeAbilities(class UGameplayAbilitySet *Set);
+	/**
+	 *	GameplayAbilities
+	 *	
+	 *	The role of the AttributeComponent wrt Abilities is to provide:
+	 *		-Management of abiltiy instances (whether per actor or per instance).
+	 *			-Someone *has* to keep track of these instances.
+	 *			-Non instanced abilities *could* be executed without any ability stuff in attribute component.
+	 *				They should be able to operate on an GameplayAbilityActorInfo + GameplayAbility.
+	 *		
+	 *	As convenience it may provide some other features:
+	 *		-Some basic input binding (whether instanced or non instanced abilities).
+	 *		-Concepts like "this component has these abilities
+	 *	
+	 */
+
+	class UGameplayAbility * CreateNewInstanceOfAbility(UGameplayAbility *Ability);
+
+	void CancelAbilitiesWithTags(const FGameplayTagContainer Tags, const FGameplayAbilityActorInfo & ActorInfo, UGameplayAbility * Ignore);
+	
+	/** References to non-replicating abilities that we instanced. We need to keep these refernces to avoid GC */
+	UPROPERTY()
+	TArray<class UGameplayAbility *>	NonReplicatedInstancedAbilities;
+
+	/** References to replicating abilities that we instanced. We need to keep these refernces to avoid GC */
+	UPROPERTY(Replicated)
+	TArray<class UGameplayAbility *>	ReplicatedInstancedAbilities;
+
+	/**
+	 *	The abilities we can activate. 
+	 *		-This will include CDOs for non instanced abilities and per-execution instanced abilities. 
+	 *		-Actor-instanced abilities will be the actual instance (not CDO)
+	 *		
+	 *	This array is not vital for things to work. It is a convenience thing for 'giving abilties to the actor'. But abilities could also work on things
+	 *	without an attribute component. For example an ability could be written to execute on a StaticMeshActor. As long as the ability doesn't require 
+	 *	instancing or anything else that the AttributeComponent would provide, then it doesn't need the component to function.
+	 */ 
+	UPROPERTY(Replicated, BlueprintReadOnly, Category = "Abilities")
+	TArray<class UGameplayAbility *>	ActivatableAbilities;
+	
 
-	/** Binds the InputComponent directly to the abilities. This is a convenience for some games. Other games may have abilities on weapons or who need some game logic to decide what ability a button press should go to. */
-	void BindInputComponentToAbilities(UInputComponent *InputComponent, UGameplayAbilitySet *AbilitySet);
+	UFUNCTION(Server, WithValidation, Reliable)
+	void	ServerTryActivateAbility(class UGameplayAbility * AbilityToActivate, int32 PredictionKey);
 
-	UPROPERTY()
-	TArray<class UGameplayAbility *>	InstancedAbilities;
+	UFUNCTION(Client, Reliable)
+	void	ClientActivateAbilityFailed(class UGameplayAbility * AbilityToActivate, int32 PredictionKey);
+
+	UFUNCTION(Client, Reliable)
+	void	ClientActivateAbilitySucceed(class UGameplayAbility * AbilityToActivate, int32 PredictionKey);
 
-	/** All abilities, including those that are shared/non instanced. */
-	UPROPERTY(BlueprintReadOnly, Category="Abilities")
-	TArray<class UGameplayAbility *>	AllAbilities;
+	// ----------------------------------------------------------------------------------------------------------------
 
 	/** This is probably temp for testing. We want to think a bit more on the API for outside stuff activating abilities */
-	UFUNCTION(BlueprintCallable, Category="Abilities")
+	UFUNCTION(BlueprintCallable, Category = "Abilities")
 	bool	ActivateAbility(int32 AbilityIdx);
 
-	// There needs to be a concept of an animating ability. Only one may exist at a time. New requestrs can be queued up, overridden, or ignored.
-
+	/** There needs to be a concept of an animating ability. Only one may exist at a time. New requestrs can be queued up, overridden, or ignored. */
 	UPROPERTY()
 	class UGameplayAbility *	AnimatingAbility;
 
@@ -175,16 +296,10 @@ class SKILLSYSTEM_API UAttributeComponent : public UActorComponent
 
 	void	MontageBranchPoint_AbilityDecisionStart();
 
-	// --- Ability RPCs
-
-	UFUNCTION(Server, WithValidation, Reliable)
-	void	ServerTryActivateAbility(class UGameplayAbility * AbilityToActivate);
-
-	UFUNCTION(Client, Reliable)
-	void	ClientActivateAbilityFailed(class UGameplayAbility * AbilityToActivate);
+	
 
-	UFUNCTION(Client, Reliable)
-	void	ClientActivateAbilitySucceed(class UGameplayAbility * AbilityToActivate);
+	// -----------------------------------------------------------------------------
+	
 	
 
 private:
@@ -242,9 +357,9 @@ protected:
 
 	virtual void OnRegister() OVERRIDE;
 
-	UAttributeSet*	GetAttributeSubobject(UClass *AttributeClass) const;
-	UAttributeSet*	GetAttributeSubobjectChecked(UClass *AttributeClass) const;
-	UAttributeSet*	GetOrCreateAttributeSubobject(UClass *AttributeClass);
+	UAttributeSet*	GetAttributeSubobject(const TSubclassOf<UAttributeSet> AttributeClass) const;
+	UAttributeSet*	GetAttributeSubobjectChecked(const TSubclassOf<UAttributeSet> AttributeClass) const;
+	UAttributeSet*	GetOrCreateAttributeSubobject(const TSubclassOf<UAttributeSet> AttributeClass);
 
 	friend struct FActiveGameplayEffectsContainer;
 	friend struct FActiveGameplayEffect;
diff --git a/Engine/Source/Runtime/SkillSystem/Classes/AttributeSet.h b/Engine/Source/Runtime/SkillSystem/Classes/AttributeSet.h
--- a/Engine/Source/Runtime/SkillSystem/Classes/AttributeSet.h
+++ b/Engine/Source/Runtime/SkillSystem/Classes/AttributeSet.h
@@ -39,7 +39,7 @@ struct SKILLSYSTEM_API FGameplayAttribute
 	void SetNumericValueChecked(const float NewValue, class UAttributeSet* Dest) const;
 
 	float GetNumericValueChecked(class UAttributeSet* Src) const;
-
+	
 	/** Equality/Inequality operators */
 	bool operator==(const FGameplayAttribute& Other) const;
 	bool operator!=(const FGameplayAttribute& Other) const;
@@ -62,13 +62,18 @@ private:
 	UProperty*	Attribute;
 };
 
-UCLASS(Blueprintable)
+UCLASS(DefaultToInstanced, Blueprintable)
 class SKILLSYSTEM_API UAttributeSet : public UObject
 {
 	GENERATED_UCLASS_BODY()
 
 public:
 
+	bool IsSupportedForNetworking() const OVERRIDE
+	{
+		return true;
+	}
+
 	/** Called just before modifying the value of an attribute. AttributeSet can make additional modifications here. */
 	virtual void PreAttributeModify(struct FGameplayEffectModCallbackData &Data) { }
 	
diff --git a/Engine/Source/Runtime/SkillSystem/Classes/BlueprintPlayMontageAndWaitTaskProxy.h b/Engine/Source/Runtime/SkillSystem/Classes/BlueprintPlayMontageAndWaitTaskProxy.h
--- a/Engine/Source/Runtime/SkillSystem/Classes/BlueprintPlayMontageAndWaitTaskProxy.h
+++ b/Engine/Source/Runtime/SkillSystem/Classes/BlueprintPlayMontageAndWaitTaskProxy.h
@@ -1,11 +1,12 @@
 // Copyright 1998-2014 Epic Games, Inc. All Rights Reserved.
 #pragma once
+#include "AbilityTaskProxyBase.h"
 #include "BlueprintPlayMontageAndWaitTaskProxy.generated.h"
 
 DECLARE_DYNAMIC_MULTICAST_DELEGATE(FMontageWaitSimpleDelegate);
 
 UCLASS(MinimalAPI)
-class UBlueprintPlayMontageAndWaitTaskProxy : public UObject
+class UBlueprintPlayMontageAndWaitTaskProxy : public UAbilityTaskProxyBase
 {
 	GENERATED_UCLASS_BODY()
 
@@ -16,7 +17,7 @@ class UBlueprintPlayMontageAndWaitTaskProxy : public UObject
 	FMontageWaitSimpleDelegate	OnInterrupted;
 
 	UFUNCTION()
-	void	OnMontageEnded(UAnimMontage* Montage, bool bInterrupted);
+	void OnMontageEnded(UAnimMontage* Montage, bool bInterrupted);
 
 public:
 	
diff --git a/Engine/Source/Runtime/SkillSystem/Classes/BlueprintWaitMovementModeChangeTaskProxy.h b/Engine/Source/Runtime/SkillSystem/Classes/BlueprintWaitMovementModeChangeTaskProxy.h
new file mode 100644
--- /dev/null
+++ b/Engine/Source/Runtime/SkillSystem/Classes/BlueprintWaitMovementModeChangeTaskProxy.h
@@ -0,0 +1,23 @@
+// Copyright 1998-2014 Epic Games, Inc. All Rights Reserved.
+#pragma once
+#include "AbilityTaskProxyBase.h"
+#include "BlueprintWaitMovementModeChangeTaskProxy.generated.h"
+
+DECLARE_DYNAMIC_MULTICAST_DELEGATE_OneParam(FMovementModeChangedDelegate, EMovementMode, NewMovementMode);
+
+UCLASS(MinimalAPI)
+class UBlueprintWaitMovementModeChangeTaskProxy : public UAbilityTaskProxyBase
+{
+	GENERATED_UCLASS_BODY()
+
+	UPROPERTY(BlueprintAssignable)
+	FMovementModeChangedDelegate	OnChange;
+
+	UFUNCTION()
+	void OnMovementModeChange(ACharacter * Character, EMovementMode PrevMovementMode, uint8 PreviousCustomMode);
+
+	EMovementMode	RequiredMode;
+
+public:
+	
+};
\ No newline at end of file
diff --git a/Engine/Source/Runtime/SkillSystem/Classes/GameplayAbility.h b/Engine/Source/Runtime/SkillSystem/Classes/GameplayAbility.h
deleted file mode 100644
diff --git a/Engine/Source/Runtime/SkillSystem/Classes/GameplayAbilitySet.h b/Engine/Source/Runtime/SkillSystem/Classes/GameplayAbilitySet.h
--- a/Engine/Source/Runtime/SkillSystem/Classes/GameplayAbilitySet.h
+++ b/Engine/Source/Runtime/SkillSystem/Classes/GameplayAbilitySet.h
@@ -27,11 +27,16 @@ struct FGameplayAbilityBindInfo
 
 	UPROPERTY(EditAnywhere, Category=BindInfo)
 	TSubclassOf<UGameplayAbility>	GameplayAbilityClass;
+
+	UPROPERTY(EditAnywhere, Category = BindInfo)
+	class UGameplayAbility * GameplayAbilityInstance;
 };
 
 /**
 * UGameplayAbilitySet
-*	This contains a list of abilities along with key bindings. THis will be very game specific, so it is expected for games to override this clas.
+*	This contains a list of abilities along with key bindings. This will be very game specific, so it is expected for games to override this class.
+*	Functions like InitializeAbilities and BindInputComponentToAbilities are setup to act on the attribute component (rather than attribute component processing the ability set).
+*	Ideally this avoids games having to implement their own UAttributeComponent and UGameplayAbilitySet, and they can just implement thei own UGameplayAbilitySet.
 */
 UCLASS(Blueprintable)
 class SKILLSYSTEM_API UGameplayAbilitySet : public UDataAsset
@@ -42,4 +47,10 @@ public:
 
 	UPROPERTY(EditAnywhere, Category=AbilitySet)
 	TArray<FGameplayAbilityBindInfo>	Abilities;
+
+	/** Give all of the abilities in this set to given AttributeComponent */
+	virtual void InitializeAbilities(UAttributeComponent *AttributeComponent) const;
+
+	/** Binds the ability set directly to this inputcomponent and attributecomponent. This is a convenience method, not all games will want to bind input directly to abilities. */
+	virtual void BindInputComponentToAbilities(UInputComponent *InputComponent, UAttributeComponent *AttributeComponent) const;
 };
\ No newline at end of file
diff --git a/Engine/Source/Runtime/SkillSystem/Classes/GameplayAbility_Instanced.h b/Engine/Source/Runtime/SkillSystem/Classes/GameplayAbility_Instanced.h
deleted file mode 100644
diff --git a/Engine/Source/Runtime/SkillSystem/Classes/GameplayAbility_Montage.h b/Engine/Source/Runtime/SkillSystem/Classes/GameplayAbility_Montage.h
deleted file mode 100644
diff --git a/Engine/Source/Runtime/SkillSystem/Classes/GameplayEffect.h b/Engine/Source/Runtime/SkillSystem/Classes/GameplayEffect.h
--- a/Engine/Source/Runtime/SkillSystem/Classes/GameplayEffect.h
+++ b/Engine/Source/Runtime/SkillSystem/Classes/GameplayEffect.h
@@ -973,6 +973,7 @@ struct FModifierQualifier
 
 	FModifierQualifier()
 		: MyType(EGameplayMod::Max)
+		, MyPredictionKey(0)
 	{
 	}
 
@@ -989,6 +990,20 @@ struct FModifierQualifier
 		return MyType;
 	}
 
+	// ----------------------------------
+	// PredictionKey is used by networking/replication to specify that a GameplayEffect has been predictively created/added.
+
+	int32 PredictionKey() const
+	{
+		return MyPredictionKey;
+	}
+
+	FModifierQualifier& PredictionKey(int32 InPredictionKey)
+	{
+		MyPredictionKey = InPredictionKey;
+		return *this;
+	}
+
 	// ----------------------------------
 	// IgnoreHandle - ignore this handle completely. For example when executing an active gameplay effect, we never
 	// want to apply it to itself (either as IncomingGE or activeGE). It is ignored in all contexts.
@@ -1045,6 +1060,7 @@ private:
 	EGameplayMod::Type	MyType;
 	FActiveGameplayEffectHandle MyIgnoreHandle;		// Do not modify this handle
 	FActiveGameplayEffectHandle MyExclusiveTargetHandle;	// Only modify this handle
+	int32 MyPredictionKey;
 };
 
 /**
@@ -1190,15 +1206,17 @@ struct FActiveGameplayEffect : public FFastArraySerializerItem
 		: StartGameStateTime(0)
 		, StartWorldTime(0.f)
 		, NextExecuteTime(0.f)
+		, PredictionKey(0)
 	{
 	}
 
-	FActiveGameplayEffect(FActiveGameplayEffectHandle InHandle, const FGameplayEffectSpec &InSpec, float CurrentWorldTime, int32 InStartGameStateTime)
+	FActiveGameplayEffect(FActiveGameplayEffectHandle InHandle, const FGameplayEffectSpec &InSpec, float CurrentWorldTime, int32 InStartGameStateTime, int32 InPredictionKey)
 		: Handle(InHandle)
 		, Spec(InSpec)
 		, StartGameStateTime(InStartGameStateTime)
 		, StartWorldTime(CurrentWorldTime)
 		, NextExecuteTime(0.f)
+		, PredictionKey(InPredictionKey)
 	{
 		// Init NextExecuteTime if necessary
 		float Period = GetPeriod();
@@ -1227,6 +1245,14 @@ struct FActiveGameplayEffect : public FFastArraySerializerItem
 
 	UPROPERTY(NotReplicated)
 	float NextExecuteTime;
+
+	int32 PredictionKey;
+
+	float GetTimeRemaining(float WorldTime)
+	{
+		float Duration = GetDuration();		
+		return (Duration == UGameplayEffect::INFINITE_DURATION ? -1.f : Duration - (WorldTime - StartWorldTime));
+	}
 	
 	float GetDuration() const
 	{
@@ -1255,36 +1281,6 @@ struct FActiveGameplayEffect : public FFastArraySerializerItem
 	}
 };
 
-
-USTRUCT()
-struct FActiveGameplayEffectData
-{
-	GENERATED_USTRUCT_BODY()
-
-	FActiveGameplayEffectData()
-		: Handle()
-		, Duration(0.f)
-		, Magnitude(0.f)
-	{
-	}
-
-	FActiveGameplayEffectData(const FActiveGameplayEffectHandle InHandle, const float InDuration, const float InMagnitude)
-		: Handle(InHandle)
-		, Duration(InDuration)
-		, Magnitude(InMagnitude)
-	{
-	}
-
-	UPROPERTY()
-	FActiveGameplayEffectHandle	Handle;
-
-	UPROPERTY()
-	float	Duration;
-
-	UPROPERTY()
-	float	Magnitude;
-};
-
 /** Generic querying data structure for active GameplayEffects. Lets us ask things like:
  *		-Give me duration/magnitude of active gameplay effects with these tags
  *		-
@@ -1325,7 +1321,7 @@ struct FActiveGameplayEffectsContainer : public FFastArraySerializer
 	UPROPERTY()
 	TArray< FActiveGameplayEffect >	GameplayEffects;
 	
-	FActiveGameplayEffect & CreateNewActiveGameplayEffect(const FGameplayEffectSpec &Spec);
+	FActiveGameplayEffect & CreateNewActiveGameplayEffect(const FGameplayEffectSpec &Spec, int32 InPredictionKey);
 
 	// returns true if none of the active effects provide immunity to Spec
 	// returns false if one (or more) of the active effects provides immunity to Spec
@@ -1378,6 +1374,8 @@ struct FActiveGameplayEffectsContainer : public FFastArraySerializer
 
 	bool HasAnyTags(FGameplayTagContainer &Tags);
 
+	bool HasAllTags(FGameplayTagContainer &Tags);
+
 	bool CanApplyAttributeModifiers(const UGameplayEffect *GameplayEffect, float Level, AActor *Instigator);
 	
 	TArray<float> GetActiveEffectsTimeRemaining(const FActiveGameplayEffectQuery Query) const;
diff --git a/Engine/Source/Runtime/SkillSystem/Classes/GameplayEffectExtension.h b/Engine/Source/Runtime/SkillSystem/Classes/GameplayEffectExtension.h
--- a/Engine/Source/Runtime/SkillSystem/Classes/GameplayEffectExtension.h
+++ b/Engine/Source/Runtime/SkillSystem/Classes/GameplayEffectExtension.h
@@ -5,6 +5,7 @@
 #include "GameplayTagContainer.h"
 #include "GameplayTagAssetInterface.h"
 #include "AttributeSet.h"
+#include "GameplayEffect.h"
 #include "GameplayEffectExtension.generated.h"
 
 struct FGameplayEffectModCallbackData
diff --git a/Engine/Source/Runtime/SkillSystem/Classes/SkillSystemBlueprintLibrary.h b/Engine/Source/Runtime/SkillSystem/Classes/SkillSystemBlueprintLibrary.h
--- a/Engine/Source/Runtime/SkillSystem/Classes/SkillSystemBlueprintLibrary.h
+++ b/Engine/Source/Runtime/SkillSystem/Classes/SkillSystemBlueprintLibrary.h
@@ -2,7 +2,7 @@
 
 #pragma once
 
-#include "GameplayAbility.h"
+#include "Abilities/GameplayAbility.h"
 #include "SkillSystemBlueprintLibrary.generated.h"
 
 class UBlueprintPlayMontageAndWaitTaskProxy;
@@ -18,4 +18,7 @@ class USkillSystemBlueprintLibrary : public UBlueprintFunctionLibrary
 
 	UFUNCTION(BlueprintCallable, meta=(HidePin="WorldContextObject", DefaultToSelf="WorldContextObject", BlueprintInternalUseOnly = "TRUE"))
 	static UBlueprintPlayMontageAndWaitTaskProxy* CreatePlayMontageAndWaitProxy(class UObject* WorldContextObject, class UAnimMontage *MontageToPlay);
+
+	UFUNCTION(BlueprintCallable, meta = (HidePin = "WorldContextObject", DefaultToSelf = "WorldContextObject", BlueprintInternalUseOnly = "TRUE"))
+	static class UBlueprintWaitMovementModeChangeTaskProxy* CreateWaitMovementModeChange(class UObject* WorldContextObject, EMovementMode NewMode);
 };
\ No newline at end of file
diff --git a/Engine/Source/Runtime/SkillSystem/Classes/SkillSystemDebugHUD.h b/Engine/Source/Runtime/SkillSystem/Classes/SkillSystemDebugHUD.h
new file mode 100644
--- /dev/null
+++ b/Engine/Source/Runtime/SkillSystem/Classes/SkillSystemDebugHUD.h
@@ -0,0 +1,43 @@
+// Copyright 1998-2014 Epic Games, Inc. All Rights Reserved.
+
+#pragma once
+
+#include "SkillSystemDebugHUD.generated.h"
+
+namespace EAlignHorizontal
+{
+	enum Type
+	{
+		Left,
+		Center,
+		Right,
+	};
+}
+
+namespace EAlignVertical
+{
+	enum Type
+	{
+		Top,
+		Center,
+		Bottom,
+	};
+}
+
+UCLASS()
+class ASkillSystemDebugHUD : public AHUD
+{
+	GENERATED_UCLASS_BODY()
+
+	/** main HUD update loop */
+	virtual void DrawHUD() OVERRIDE;
+
+	void DrawWithBackground(UFont* InFont, const FString& Text, const FColor& TextColor, EAlignHorizontal::Type HAlign, float& OffsetX, EAlignVertical::Type VAlign, float& OffsetY, float Alpha=1.f);
+
+	void DrawDebugHUD(UCanvas* Canvas, APlayerController* PC);
+
+private:
+
+	void DrawDebugAttributeComponent(UAttributeComponent *Component);
+
+};
\ No newline at end of file
diff --git a/Engine/Source/Runtime/SkillSystem/Private/Abilities/GameplayAbility.cpp b/Engine/Source/Runtime/SkillSystem/Private/Abilities/GameplayAbility.cpp
new file mode 100644
--- /dev/null
+++ b/Engine/Source/Runtime/SkillSystem/Private/Abilities/GameplayAbility.cpp
@@ -0,0 +1,331 @@
+// Copyright 1998-2013 Epic Games, Inc. All Rights Reserved.
+
+#include "SkillSystemModulePrivatePCH.h"
+
+// --------------------------------------------------------------------------------------------------------------------------------------------------------
+//
+//	UGameplayAbility
+//
+// --------------------------------------------------------------------------------------------------------------------------------------------------------
+
+UGameplayAbility::UGameplayAbility(const class FPostConstructInitializeProperties& PCIP)
+	: Super(PCIP)
+{
+	CostGameplayEffect = NULL;
+	CooldownGameplayEffect = NULL;
+
+}
+
+bool UGameplayAbility::IsSupportedForNetworking() const
+{
+	/**
+	 *	If this check fails, it means we are trying to serialize a reference to an invalid GameplayAbility. 
+	 *	We can only replicate references to:
+	 *		-CDOs and DataAssets (e.g., static, non-instanced gameplay abilities)
+	 *		-Instanced abilities that are replicating (and will thus be created on clients).
+	 *		
+	 *	The network code should never be asking a gameplay ability if it IsSupportedForNetworking() for 
+	 *	an ability that is not described above.
+	 */
+
+	bool Supported = GetReplicationPolicy() != EGameplayAbilityReplicationPolicy::ReplicateNone || GetOuter()->IsA(UPackage::StaticClass());
+	check(Supported);
+
+	return Supported;
+}
+
+bool UGameplayAbility::CanActivateAbility(const FGameplayAbilityActorInfo ActorInfo) const
+{
+	// Check cooldowns and resource consumption immediately
+	if (!CheckCooldown(ActorInfo))
+	{
+		return false;
+	}
+
+	if (!CheckCost(ActorInfo))
+	{
+		return false;
+	}
+
+	return true;
+}
+
+bool UGameplayAbility::CommitAbility(const FGameplayAbilityActorInfo ActorInfo)
+{
+	if (!CommitCheck(ActorInfo))
+	{
+		return false;
+	}
+
+	CommitExecute(ActorInfo);
+	return true;
+}
+
+bool UGameplayAbility::CommitCheck(const FGameplayAbilityActorInfo ActorInfo)
+{
+	return CanActivateAbility(ActorInfo);
+}
+
+void UGameplayAbility::CommitExecute(const FGameplayAbilityActorInfo ActorInfo)
+{
+	if (!ActorInfo.AttributeComponent->IsNetSimulating() || ActorInfo.PredictionKey != 0)
+	{
+		ApplyCooldown(ActorInfo);
+
+		ApplyCost(ActorInfo);
+	}
+}
+
+bool UGameplayAbility::TryActivateAbility(const FGameplayAbilityActorInfo ActorInfo, UGameplayAbility ** OutInstancedAbility)
+{	
+	// This should only come from button presses/local instigation
+	ENetRole NetMode = ActorInfo.Actor->Role;
+	ensure(NetMode != ROLE_SimulatedProxy);
+
+	// Always do a non instanced CanActivate check
+	if (!CanActivateAbility(ActorInfo))
+	{
+		return false;
+	}
+
+	// If we are the server or this is a client authorative 
+	if (NetExecutionPolicy == EGameplayAbilityNetExecutionPolicy::Client || (NetMode == ROLE_Authority))
+	{
+		// Create instance of this ability if necessary
+		
+		if (GetInstancingPolicy() == EGameplayAbilityInstancingPolicy::InstancedPerExecution)
+		{
+			UGameplayAbility * Ability = ActorInfo.AttributeComponent->CreateNewInstanceOfAbility(this);
+			Ability->CallActivateAbility(ActorInfo);
+			if (OutInstancedAbility)
+			{
+				*OutInstancedAbility = Ability;
+			}
+		}
+		else
+		{
+			CallActivateAbility(ActorInfo);
+		}
+	}
+	else if (NetExecutionPolicy == EGameplayAbilityNetExecutionPolicy::Server)
+	{
+		ServerTryActivateAbility(ActorInfo);
+	}
+	else if (NetExecutionPolicy == EGameplayAbilityNetExecutionPolicy::Predictive)
+	{
+		// This execution is now officially EGameplayAbilityActivationMode:Predicting and has a PredictionKey
+		ActorInfo.GeneratePredictionKey();
+
+		if (GetInstancingPolicy() == EGameplayAbilityInstancingPolicy::InstancedPerExecution)
+		{
+			// For now, only NonReplicated + InstancedPerExecution abilities can be Predictive.
+			// We lack the code to predict spawning an instance of the execution and then merge/combine
+			// with the server spawned version when it arrives.
+			
+			if(GetReplicationPolicy() == EGameplayAbilityReplicationPolicy::ReplicateNone)
+			{
+				UGameplayAbility * Ability = ActorInfo.AttributeComponent->CreateNewInstanceOfAbility(this);
+				Ability->CallPredictiveActivateAbility(ActorInfo);
+				if (OutInstancedAbility)
+				{
+					*OutInstancedAbility = Ability;
+				}
+			}
+			else
+			{
+				ensure(false);
+			}
+		}
+		else
+		{
+			CallPredictiveActivateAbility(ActorInfo);
+		}
+		
+		ServerTryActivateAbility(ActorInfo);
+	}
+
+	return true;	
+}
+
+void UGameplayAbility::ActivateAbility(FGameplayAbilityActorInfo ActorInfo)
+{
+	// Child classes may want to do stuff here...
+
+	if (CommitAbility(ActorInfo))		// ..then commit the ability...
+	{
+		//	Then do more stuff...
+	}
+}
+
+void UGameplayAbility::PreActivate(const FGameplayAbilityActorInfo ActorInfo)
+{
+	ActorInfo.AttributeComponent->CancelAbilitiesWithTags(CancelAbilitiesWithTag, ActorInfo, this);
+}
+
+void UGameplayAbility::CallActivateAbility(const FGameplayAbilityActorInfo ActorInfo)
+{
+	PreActivate(ActorInfo);
+	ActivateAbility(ActorInfo);
+}
+
+void UGameplayAbility::CallPredictiveActivateAbility(const FGameplayAbilityActorInfo ActorInfo)
+{
+	PreActivate(ActorInfo);
+	PredictiveActivateAbility(ActorInfo);
+}
+
+void UGameplayAbility::PredictiveActivateAbility(const FGameplayAbilityActorInfo ActorInfo)
+{
+	
+}
+
+void UGameplayAbility::ServerTryActivateAbility(const FGameplayAbilityActorInfo ActorInfo)
+{
+	ActorInfo.AttributeComponent->ServerTryActivateAbility(this, ActorInfo.PredictionKey);
+}
+
+void UGameplayAbility::ClientActivateAbility(const FGameplayAbilityActorInfo ActorInfo)
+{
+	CallActivateAbility(ActorInfo);
+}
+
+void UGameplayAbility::InputPressed(int32 InputID, const FGameplayAbilityActorInfo ActorInfo)
+{
+	TryActivateAbility(ActorInfo);
+}
+
+void UGameplayAbility::InputReleased(int32 InputID, const FGameplayAbilityActorInfo ActorInfo)
+{
+	SKILL_LOG(Log, TEXT("InputReleased: %d"), InputID);
+}
+
+bool UGameplayAbility::CheckCooldown(const FGameplayAbilityActorInfo ActorInfo) const
+{
+	if (CooldownGameplayEffect)
+	{
+		check(ActorInfo.AttributeComponent.IsValid());
+		if (ActorInfo.AttributeComponent->HasAnyTags(CooldownGameplayEffect->OwnedTagsContainer))
+		{
+			return false;
+		}
+	}
+	return true;
+}
+
+void UGameplayAbility::ApplyCooldown(const FGameplayAbilityActorInfo ActorInfo)
+{
+	if (CooldownGameplayEffect)
+	{
+		check(ActorInfo.AttributeComponent.IsValid());
+		ActorInfo.AttributeComponent->ApplyGameplayEffectToSelf(CooldownGameplayEffect, 1.f, ActorInfo.Actor.Get(), FModifierQualifier().PredictionKey(ActorInfo.PredictionKey));
+	}
+}
+
+bool UGameplayAbility::CheckCost(const FGameplayAbilityActorInfo ActorInfo) const
+{
+	if (CostGameplayEffect)
+	{
+		check(ActorInfo.AttributeComponent.IsValid());
+		return ActorInfo.AttributeComponent->CanApplyAttributeModifiers(CostGameplayEffect, 1.f, ActorInfo.Actor.Get());
+	}
+	return true;
+}
+
+void UGameplayAbility::ApplyCost(const FGameplayAbilityActorInfo ActorInfo)
+{
+	check(ActorInfo.AttributeComponent.IsValid());
+	if (CostGameplayEffect)
+	{
+		ActorInfo.AttributeComponent->ApplyGameplayEffectToSelf(CostGameplayEffect, 1.f, ActorInfo.Actor.Get(), FModifierQualifier().PredictionKey(ActorInfo.PredictionKey));
+	}
+}
+
+float UGameplayAbility::GetCooldownTimeRemaining(const FGameplayAbilityActorInfo ActorInfo) const
+{
+	check(ActorInfo.AttributeComponent.IsValid());
+	if (CooldownGameplayEffect)
+	{
+		TArray< float > Durations = ActorInfo.AttributeComponent->GetActiveEffectsTimeRemaining(FActiveGameplayEffectQuery(&CooldownGameplayEffect->OwnedTagsContainer));
+		if (Durations.Num() > 0)
+		{
+			Durations.Sort();
+			return Durations[0];
+		}
+	}
+
+	return 0.f;
+}
+
+// --------------------------------------------------------------------
+
+void UGameplayAbility::MontageBranchPoint_AbilityDecisionStop(const FGameplayAbilityActorInfo ActorInfo) const
+{
+
+}
+
+void UGameplayAbility::MontageBranchPoint_AbilityDecisionStart(const FGameplayAbilityActorInfo ActorInfo) const
+{
+
+}
+
+void UGameplayAbility::ClientActivateAbilitySucceed_Implementation(AActor *Actor)
+{
+	if (Actor)
+	{
+		FGameplayAbilityActorInfo	ActorInfo;
+		ActorInfo.InitFromActor(Actor);
+
+		CallActivateAbility(ActorInfo);
+	}
+	else
+	{
+		SKILL_LOG(Warning, TEXT("ClientActivateAbilitySucceed called on %s with no valid Actor"), *GetName());
+	}
+}
+
+//----------------------------------------------------------------------
+
+void FGameplayAbilityActorInfo::InitFromActor(AActor *InActor)
+{
+	Actor = InActor;
+
+	// On Init, we are either Authority or NonAuthority. We haven't been given a PredictionKey and we haven't been confirmed.
+	// NonAuthority essentially means 'I'm not sure what how I'm going to do this yet'.
+	// 
+	ActivationMode = (Actor->Role == ROLE_Authority ? EGameplayAbilityActivationMode::Authority : EGameplayAbilityActivationMode::NonAuthority);
+	PredictionKey = 0;
+
+	// Grab Components that we care about
+	USkeletalMeshComponent * SkelMeshComponent = InActor->FindComponentByClass<USkeletalMeshComponent>();
+	if (SkelMeshComponent)
+	{
+		this->AnimInstance = SkelMeshComponent->GetAnimInstance();
+	}
+
+	AttributeComponent = InActor->FindComponentByClass<UAttributeComponent>();
+
+	MovementComponent = InActor->FindComponentByClass<UMovementComponent>();
+}
+
+void FGameplayAbilityActorInfo::GeneratePredictionKey() const
+{
+	check(AttributeComponent.IsValid());
+	check(PredictionKey == 0);
+	check(ActivationMode == EGameplayAbilityActivationMode::NonAuthority);
+
+	PredictionKey = AttributeComponent->GetNextPredictionKey();
+	ActivationMode = EGameplayAbilityActivationMode::Predicting;
+}
+
+void FGameplayAbilityActorInfo::SetPredictionKey(int32 InPredictionKey)
+{
+	PredictionKey = InPredictionKey;
+}
+
+void FGameplayAbilityActorInfo::SetActivationConfirmed()
+{
+	check(ActivationMode == EGameplayAbilityActivationMode::NonAuthority);
+	ActivationMode = EGameplayAbilityActivationMode::Confirmed;
+}
+
+//----------------------------------------------------------------------
\ No newline at end of file
diff --git a/Engine/Source/Runtime/SkillSystem/Private/Abilities/GameplayAbility_CharacterJump.cpp b/Engine/Source/Runtime/SkillSystem/Private/Abilities/GameplayAbility_CharacterJump.cpp
new file mode 100644
--- /dev/null
+++ b/Engine/Source/Runtime/SkillSystem/Private/Abilities/GameplayAbility_CharacterJump.cpp
@@ -0,0 +1,71 @@
+// Copyright 1998-2014 Epic Games, Inc. All Rights Reserved.
+
+#include "SkillSystemModulePrivatePCH.h"
+
+// --------------------------------------------------------------------------------------------------------------------------------------------------------
+//
+//	UGameplayAbility_CharacterJump
+//
+// --------------------------------------------------------------------------------------------------------------------------------------------------------
+
+UGameplayAbility_CharacterJump::UGameplayAbility_CharacterJump(const class FPostConstructInitializeProperties& PCIP)
+	: Super(PCIP)
+{
+	NetExecutionPolicy = EGameplayAbilityNetExecutionPolicy::Predictive;
+}
+
+void UGameplayAbility_CharacterJump::ActivateAbility(const FGameplayAbilityActorInfo ActorInfo)
+{
+	if (ActorInfo.ActivationMode == EGameplayAbilityActivationMode::Authority ||
+		ActorInfo.ActivationMode == EGameplayAbilityActivationMode::Predicting)
+	{
+		if (!CommitAbility(ActorInfo))
+		{
+			return;
+		}
+
+		SKILL_LOG(Warning, TEXT("ActivateAbility Jump. %s. ActivationMode: %d"), *GetPathName(), (int32)ActorInfo.ActivationMode);
+
+		ACharacter * Character = CastChecked<ACharacter>(ActorInfo.Actor.Get());
+		Character->Jump();
+	}
+}
+
+void UGameplayAbility_CharacterJump::PredictiveActivateAbility(const FGameplayAbilityActorInfo ActorInfo)
+{
+	CallActivateAbility(ActorInfo);
+}
+
+void UGameplayAbility_CharacterJump::InputReleased(int32 InputID, const FGameplayAbilityActorInfo ActorInfo)
+{
+	CancelAbility(ActorInfo);
+}
+
+bool UGameplayAbility_CharacterJump::CanActivateAbility(const FGameplayAbilityActorInfo ActorInfo) const
+{
+	if (!Super::CanActivateAbility(ActorInfo))
+	{
+		return false;
+	}
+
+	const UCharacterMovementComponent *CharMovement = CastChecked<UCharacterMovementComponent>(ActorInfo.MovementComponent.Get());
+	return (CharMovement->MovementMode == EMovementMode::MOVE_Walking);
+}
+
+/**
+ *	Canceling an non instanced ability is tricky. Right now this works for Jump since there is nothing that can go wrong by calling
+ *	StopJumping() if you aren't already jumping. If we had a montage playing non instanced ability, it would need to make sure the
+ *	Montage that *it* played was still playing, and if so, to cancel it. If this is something we need to support, we may need some
+ *	light weight data structure to represent 'non intanced abilities in action' with a way to cancel/end them.
+ */
+void UGameplayAbility_CharacterJump::CancelAbility(const FGameplayAbilityActorInfo ActorInfo)
+{
+	if (ActorInfo.ActivationMode == EGameplayAbilityActivationMode::Authority ||
+		ActorInfo.ActivationMode == EGameplayAbilityActivationMode::Predicting)
+	{
+		SKILL_LOG(Warning, TEXT("CancelAbility Jump. %s. ActivationMode: %d"), *GetPathName(), (int32)ActorInfo.ActivationMode);
+
+		ACharacter * Character = CastChecked<ACharacter>(ActorInfo.Actor.Get());
+		Character->StopJumping();
+	}
+}
\ No newline at end of file
diff --git a/Engine/Source/Runtime/SkillSystem/Private/Abilities/GameplayAbility_Instanced.cpp b/Engine/Source/Runtime/SkillSystem/Private/Abilities/GameplayAbility_Instanced.cpp
new file mode 100644
--- /dev/null
+++ b/Engine/Source/Runtime/SkillSystem/Private/Abilities/GameplayAbility_Instanced.cpp
@@ -0,0 +1,113 @@
+// Copyright 1998-2013 Epic Games, Inc. All Rights Reserved.
+
+#include "SkillSystemModulePrivatePCH.h"
+
+// --------------------------------------------------------------------------------------------------------------------------------------------------------
+//
+//	UGameplayAbility_Instanced
+//
+// --------------------------------------------------------------------------------------------------------------------------------------------------------
+
+
+UGameplayAbility_Instanced::UGameplayAbility_Instanced(const class FPostConstructInitializeProperties& PCIP)
+	: Super(PCIP)
+{
+	InstancedPerExecution = true;
+	ReplicationPolicy = EGameplayAbilityReplicationPolicy::ReplicateOwner;
+}
+
+bool UGameplayAbility_Instanced::CanActivateAbility(const FGameplayAbilityActorInfo ActorInfo) const
+{
+	if (!Super::CanActivateAbility(ActorInfo))
+	{
+		return false;
+	}
+
+	// Call into blueprint (fixme: cache this)
+	static FName FuncName = FName(TEXT("K2_CanActivateAbility"));
+	UFunction* CanActivateFunction = GetClass()->FindFunctionByName(FuncName);
+
+	if (CanActivateFunction
+		&& CanActivateFunction->GetOuter()->IsA(UBlueprintGeneratedClass::StaticClass())
+		&& K2_CanActivateAbility(ActorInfo) == false)
+	{
+		SKILL_LOG(Log, TEXT("CanActivateAbility %s failed, blueprint refused"), *GetName());
+		return false;
+	}
+
+	return true;
+}
+
+void UGameplayAbility_Instanced::CommitExecute(const FGameplayAbilityActorInfo ActorInfo)
+{
+	Super::CommitExecute(ActorInfo);
+
+	check(!HasAnyFlags(RF_ClassDefaultObject));
+
+	// Call into blueprint (fixme: cache this)
+	static FName FuncName = FName(TEXT("K2_CommitAbility"));
+	UFunction* CanActivateFunction = GetClass()->FindFunctionByName(FuncName);
+
+	if (CanActivateFunction	&& CanActivateFunction->GetOuter()->IsA(UBlueprintGeneratedClass::StaticClass()))
+	{
+		K2_CommitExecute(ActorInfo);
+	}
+}
+
+void UGameplayAbility_Instanced::ActivateAbility(FGameplayAbilityActorInfo ActorInfo)
+{
+	//check(!HasAnyFlags(RF_ClassDefaultObject));
+
+	// Call into blueprint (fixme: cache this)
+	static FName FuncName = FName(TEXT("K2_ActivateAbility"));
+	UFunction* CanActivateFunction = GetClass()->FindFunctionByName(FuncName);
+
+	if (CanActivateFunction	&& CanActivateFunction->GetOuter()->IsA(UBlueprintGeneratedClass::StaticClass()))
+	{
+		K2_ActivateAbility(ActorInfo);
+	}
+}
+
+void UGameplayAbility_Instanced::PredictiveActivateAbility(const FGameplayAbilityActorInfo ActorInfo)
+{
+	Super::PredictiveActivateAbility(ActorInfo);
+
+	// FIXME
+	if(!HasAnyFlags(RF_ClassDefaultObject))
+	{
+		// Call into blueprint (fixme: cache this)
+		static FName FuncName = FName(TEXT("K2_PredictiveActivateAbility"));
+		UFunction* PredictiveActivateFunction = GetClass()->FindFunctionByName(FuncName);
+
+		if (PredictiveActivateFunction&& PredictiveActivateFunction->GetOuter()->IsA(UBlueprintGeneratedClass::StaticClass()))
+		{
+			K2_PredictiveActivateAbility(ActorInfo);
+		}
+	}
+}
+
+int32 UGameplayAbility_Instanced::GetFunctionCallspace(UFunction* Function, void* Parameters, FFrame* Stack)
+{
+	if(HasAnyFlags(RF_ClassDefaultObject))
+	{
+		return FunctionCallspace::Local;
+	}
+	check(GetOuter() != NULL);
+	return GetOuter()->GetFunctionCallspace(Function, Parameters, Stack);
+}
+
+bool UGameplayAbility_Instanced::CallRemoteFunction(UFunction* Function, void* Parameters, FOutParmRec* OutParms, FFrame* Stack)
+{
+	check(!HasAnyFlags(RF_ClassDefaultObject));
+	check(GetOuter() != NULL);
+
+	AActor* Owner = CastChecked<AActor>(GetOuter());
+	UNetDriver* NetDriver = Owner->GetNetDriver();
+	if (NetDriver)
+	{
+		NetDriver->ProcessRemoteFunction(Owner, Function, Parameters, OutParms, Stack, this);
+		return true;
+	}
+
+	return false;
+}
\ No newline at end of file
diff --git a/Engine/Source/Runtime/SkillSystem/Private/Abilities/GameplayAbility_Montage.cpp b/Engine/Source/Runtime/SkillSystem/Private/Abilities/GameplayAbility_Montage.cpp
new file mode 100644
--- /dev/null
+++ b/Engine/Source/Runtime/SkillSystem/Private/Abilities/GameplayAbility_Montage.cpp
@@ -0,0 +1,63 @@
+// Copyright 1998-2013 Epic Games, Inc. All Rights Reserved.
+
+#include "SkillSystemModulePrivatePCH.h"
+
+// --------------------------------------------------------------------------------------------------------------------------------------------------------
+//
+//	UGameplayAbility_Montage
+//
+// --------------------------------------------------------------------------------------------------------------------------------------------------------
+
+UGameplayAbility_Montage::UGameplayAbility_Montage(const class FPostConstructInitializeProperties& PCIP)
+	: Super(PCIP)
+{
+	PlayRate = 1.f;
+	NetExecutionPolicy = EGameplayAbilityNetExecutionPolicy::Server;
+}
+
+void UGameplayAbility_Montage::ActivateAbility(const FGameplayAbilityActorInfo OwnerInfo)
+{
+	if (!CommitAbility(OwnerInfo))
+	{
+		return;
+	}
+
+	if (MontageToPlay != NULL && OwnerInfo.AnimInstance != NULL && OwnerInfo.AnimInstance->GetActiveMontageInstance() == NULL)
+	{
+		TArray<FActiveGameplayEffectHandle>	AppliedEffects;
+
+		// Apply GameplayEffects
+		for (const UGameplayEffect * Effect : GameplayEffectsWhileAnimating)
+		{
+			FActiveGameplayEffectHandle Handle = OwnerInfo.AttributeComponent->ApplyGameplayEffectToSelf(Effect, 1.f, OwnerInfo.Actor.Get());
+			if (Handle.IsValid())
+			{
+				AppliedEffects.Add(Handle);
+			}
+		}
+
+		float const Duration = OwnerInfo.AnimInstance->Montage_Play(MontageToPlay, PlayRate);
+
+		FOnMontageEnded EndDelegate;
+		EndDelegate.BindUObject(this, &UGameplayAbility_Montage::OnMontageEnded, OwnerInfo, AppliedEffects);
+		OwnerInfo.AnimInstance->Montage_SetEndDelegate(EndDelegate);
+
+		if (SectionName != NAME_None)
+		{
+			OwnerInfo.AnimInstance->Montage_JumpToSection(SectionName);
+		}
+	}
+}
+
+void UGameplayAbility_Montage::OnMontageEnded(UAnimMontage* Montage, bool bInterrupted, const FGameplayAbilityActorInfo OwnerInfo, TArray<FActiveGameplayEffectHandle> AppliedEffects)
+{
+	// Remove any GameplayEffects that we applied
+	UAttributeComponent *OwnerComponent = OwnerInfo.AttributeComponent.Get();
+	if (OwnerComponent)
+	{
+		for (FActiveGameplayEffectHandle Handle : AppliedEffects)
+		{
+			OwnerComponent->RemoveActiveGameplayEffect(Handle);
+		}
+	}
+}
\ No newline at end of file
diff --git a/Engine/Source/Runtime/SkillSystem/Private/AbilityTaskProxyBase.cpp b/Engine/Source/Runtime/SkillSystem/Private/AbilityTaskProxyBase.cpp
new file mode 100644
--- /dev/null
+++ b/Engine/Source/Runtime/SkillSystem/Private/AbilityTaskProxyBase.cpp
@@ -0,0 +1,8 @@
+
+#include "SkillSystemModulePrivatePCH.h"
+
+UAbilityTaskProxyBase::UAbilityTaskProxyBase(const class FPostConstructInitializeProperties& PCIP)
+	: Super(PCIP)
+{
+
+}
\ No newline at end of file
diff --git a/Engine/Source/Runtime/SkillSystem/Private/AttributeComponent.cpp b/Engine/Source/Runtime/SkillSystem/Private/AttributeComponent.cpp
--- a/Engine/Source/Runtime/SkillSystem/Private/AttributeComponent.cpp
+++ b/Engine/Source/Runtime/SkillSystem/Private/AttributeComponent.cpp
@@ -27,6 +27,8 @@ UAttributeComponent::UAttributeComponent(const class FPostConstructInitializePro
 	ActiveGameplayEffects.Owner = this;
 
 	SetIsReplicated(true);
+	ReplicatedPredictionKey = 0;
+	LocalPredictionKey = 0;
 }
 
 UAttributeComponent::~UAttributeComponent()
@@ -61,7 +63,7 @@ void UAttributeComponent::ModifyStats(TSubclassOf<class UAttributeSet> Attribute
 	}
 }
 
-UAttributeSet*	UAttributeComponent::GetOrCreateAttributeSubobject(UClass *AttributeClass)
+UAttributeSet*	UAttributeComponent::GetOrCreateAttributeSubobject(const TSubclassOf<UAttributeSet> AttributeClass)
 {
 	AActor *OwningActor = GetOwner();
 	UAttributeSet *MyAttributes  = NULL;
@@ -72,41 +74,28 @@ UAttributeSet*	UAttributeComponent::GetOrCreateAttributeSubobject(UClass *Attrib
 		{
 			MyAttributes = ConstructObject<UAttributeSet>(AttributeClass, OwningActor);
 			SpawnedAttributes.AddUnique(MyAttributes);
-			//LastSpawnedSet = MyAttributes;
-			TestFloat= 999.f;
 		}
 	}
 
 	return MyAttributes;
 }
 
-UAttributeSet* UAttributeComponent::GetAttributeSubobjectChecked(UClass *AttributeClass) const
+UAttributeSet* UAttributeComponent::GetAttributeSubobjectChecked(const TSubclassOf<UAttributeSet> AttributeClass) const
 {
 	UAttributeSet *Set = GetAttributeSubobject(AttributeClass);
 	check(Set);
 	return Set;
 }
 
-UAttributeSet* UAttributeComponent::GetAttributeSubobject(UClass *AttributeClass) const
+UAttributeSet* UAttributeComponent::GetAttributeSubobject(const TSubclassOf<UAttributeSet> AttributeClass) const
 {
-	// FIXME: Lookup may would be helpful here
-
-	AActor *ActorOwner = GetOwner();
-	if (ActorOwner)
+	for (UAttributeSet * Set : SpawnedAttributes)
 	{
-		TArray<UObject*> ChildObjects;
-		GetObjectsWithOuter(ActorOwner, ChildObjects);
-
-		for (int32 ChildIndex=0; ChildIndex < ChildObjects.Num(); ++ChildIndex)
+		if (Set && Set->IsA(AttributeClass))
 		{
-			UObject *Obj = ChildObjects[ChildIndex];
-			if (Obj->GetClass()->IsChildOf(AttributeClass))
-			{
-				return Cast<UAttributeSet>(Obj);
-			}
+			return Set;
 		}
 	}
-
 	return NULL;
 }
 
@@ -213,6 +202,11 @@ bool UAttributeComponent::HasAnyTags(FGameplayTagContainer &Tags)
 	return ActiveGameplayEffects.HasAnyTags(Tags);
 }
 
+bool UAttributeComponent::HasAllTags(FGameplayTagContainer &Tags)
+{
+	return ActiveGameplayEffects.HasAllTags(Tags);
+}
+
 void UAttributeComponent::TEMP_ApplyActiveGameplayEffects()
 {
 	for (int32 idx=0; idx < ActiveGameplayEffects.GameplayEffects.Num(); ++idx)
@@ -239,6 +233,9 @@ FActiveGameplayEffectHandle UAttributeComponent::ApplyGameplayEffectSpecToTarget
 
 FActiveGameplayEffectHandle UAttributeComponent::ApplyGameplayEffectSpecToSelf(const FGameplayEffectSpec &Spec, FModifierQualifier BaseQualifier)
 {
+	// Temp, only non instance, non periodic GEs can be predictive
+	check((BaseQualifier.PredictionKey() >= 0) == (Spec.GetDuration() != UGameplayEffect::INSTANT_APPLICATION && Spec.GetPeriod() == UGameplayEffect::NO_PERIOD));
+
 	// check if the effect being applied actually succeeds
 	float ChanceToApply = Spec.GetChanceToApplyToTarget();
 	if ((ChanceToApply < 1.f - SMALL_NUMBER) && (FMath::FRand() > ChanceToApply))
@@ -257,7 +254,7 @@ FActiveGameplayEffectHandle UAttributeComponent::ApplyGameplayEffectSpecToSelf(c
 		float Duration = Spec.GetDuration();
 		if (Duration != UGameplayEffect::INSTANT_APPLICATION)
 		{
-			FActiveGameplayEffect &NewActiveEffect = ActiveGameplayEffects.CreateNewActiveGameplayEffect(Spec);
+			FActiveGameplayEffect &NewActiveEffect = ActiveGameplayEffects.CreateNewActiveGameplayEffect(Spec, BaseQualifier.PredictionKey());
 			MyHandle = NewActiveEffect.Handle;
 			OurCopyOfSpec = &NewActiveEffect.Spec;
 
@@ -454,10 +451,73 @@ void UAttributeComponent::InvokeGameplayCueRemoved(const FGameplayEffectSpec &Sp
 	}
 }
 
+void UAttributeComponent::AddDependancyToAttribute(FGameplayAttribute Attribute, const TWeakPtr<FAggregator> InDependant)
+{
+	ActiveGameplayEffects.AddDependancyToAttribute(Attribute, InDependant);
+}
+
+bool UAttributeComponent::CanApplyAttributeModifiers(const UGameplayEffect *GameplayEffect, float Level, AActor *Instigator)
+{
+	return ActiveGameplayEffects.CanApplyAttributeModifiers(GameplayEffect, Level, Instigator);
+}
+
+TArray<float> UAttributeComponent::GetActiveEffectsTimeRemaining(const FActiveGameplayEffectQuery Query) const
+{
+	return ActiveGameplayEffects.GetActiveEffectsTimeRemaining(Query);
+}
+
+// ---------------------------------------------------------------------------------------
+
 void UAttributeComponent::GetLifetimeReplicatedProps(TArray< FLifetimeProperty > & OutLifetimeProps) const
 {
 	Super::GetLifetimeReplicatedProps(OutLifetimeProps);
+
+	DOREPLIFETIME(UAttributeComponent, SpawnedAttributes);
 	DOREPLIFETIME(UAttributeComponent, ActiveGameplayEffects);
+
+	DOREPLIFETIME(UAttributeComponent, ReplicatedInstancedAbilities);
+	DOREPLIFETIME(UAttributeComponent, ActivatableAbilities);
+
+	DOREPLIFETIME(UAttributeComponent, ReplicatedPredictionKey);
+}
+
+bool UAttributeComponent::ReplicateSubobjects(class UActorChannel *Channel, class FOutBunch *Bunch, FReplicationFlags *RepFlags)
+{
+	bool WroteSomething = Super::ReplicateSubobjects(Channel, Bunch, RepFlags);
+
+	for (UAttributeSet* Set : SpawnedAttributes)
+	{
+		if (Set)
+		{
+			WroteSomething |= Channel->ReplicateSubobject(Set, *Bunch, *RepFlags);
+		}
+	}
+
+	for (UGameplayAbility* Ability : ReplicatedInstancedAbilities)
+	{
+		if (Ability)
+		{
+			WroteSomething |= Channel->ReplicateSubobject(Ability, *Bunch, *RepFlags);
+		}
+	}
+
+	return WroteSomething;
+}
+
+void UAttributeComponent::GetSubobjectsWithStableNamesForNetworking(TArray<UObject*>& Objs)
+{
+	for (UAttributeSet* Set : SpawnedAttributes)
+	{
+		if (Set && Set->IsNameStableForNetworking())
+		{
+			Objs.Add(Set);
+		}
+	}
+}
+
+void UAttributeComponent::PostNetReceive()
+{
+	Super::PostNetReceive();
 }
 
 void UAttributeComponent::OnRep_GameplayEffects()
@@ -465,19 +525,50 @@ void UAttributeComponent::OnRep_GameplayEffects()
 
 }
 
-void UAttributeComponent::AddDependancyToAttribute(FGameplayAttribute Attribute, const TWeakPtr<FAggregator> InDependant)
+void UAttributeComponent::OnRep_PredictionKey()
 {
-	ActiveGameplayEffects.AddDependancyToAttribute(Attribute, InDependant);
+	// Every predictive action we've done up to and including the current value of ReplicatedPredictionKey needs to be wiped
+	int32 idx = 0;
+	for ( ; idx < PredictionDelegates.Num(); ++idx)
+	{
+		if (PredictionDelegates[idx].Key <= ReplicatedPredictionKey)
+		{
+			PredictionDelegates[idx].Value.Broadcast();
+		}
+		else
+		{
+			break;
+		}
+	}
+	if (idx > 0)
+	{
+		// Remove everyone we called Broadcast on in one call
+		PredictionDelegates.RemoveAt(0, idx);
+	}
 }
 
-bool UAttributeComponent::CanApplyAttributeModifiers(const UGameplayEffect *GameplayEffect, float Level, AActor *Instigator)
+FAttributeComponentPredictionKeyClear& UAttributeComponent::GetPredictionKeyDelegate(int32 PredictionKey)
 {
-	return ActiveGameplayEffects.CanApplyAttributeModifiers(GameplayEffect, Level, Instigator);
+	// See if we already have one for this key
+	for (TPair<int32, FAttributeComponentPredictionKeyClear> & Item : PredictionDelegates)
+	{
+		if (Item.Key == PredictionKey)
+		{
+			return Item.Value;
+		}
+	}
+
+	// Create a new one
+	TPair<int32, FAttributeComponentPredictionKeyClear> NewPair;
+	NewPair.Key = PredictionKey;
+	
+	PredictionDelegates.Add(NewPair);
+	return PredictionDelegates.Last().Value;
 }
 
-TArray<float> UAttributeComponent::GetActiveEffectsTimeRemaining(const FActiveGameplayEffectQuery Query) const
+int32 UAttributeComponent::GetNextPredictionKey()
 {
-	return ActiveGameplayEffects.GetActiveEffectsTimeRemaining(Query);
+	return ++LocalPredictionKey;
 }
 
 // ---------------------------------------------------------------------------------------
diff --git a/Engine/Source/Runtime/SkillSystem/Private/AttributeComponent_Abilities.cpp b/Engine/Source/Runtime/SkillSystem/Private/AttributeComponent_Abilities.cpp
--- a/Engine/Source/Runtime/SkillSystem/Private/AttributeComponent_Abilities.cpp
+++ b/Engine/Source/Runtime/SkillSystem/Private/AttributeComponent_Abilities.cpp
@@ -16,95 +16,48 @@
 void UAttributeComponent::InitializeComponent()
 {
 	Super::InitializeComponent();
-}
-void UAttributeComponent::InitializeAbilities(class UGameplayAbilitySet *AbilitySet)
-{
-	AActor * OwnerActor = GetOwner();
-	check(OwnerActor);
-
-	FGameplayAbilityActorInfo	ActorInfo;
-	ActorInfo.InitFromActor(OwnerActor);
+	
+	// Look for DSO AttributeSets
+	AActor *Owner = GetOwner();
 
-	for (const FGameplayAbilityBindInfo &BindInfo : AbilitySet->Abilities)
+	TArray<UObject*> ChildObjects;
+	GetObjectsWithOuter(Owner, ChildObjects, false, RF_PendingKill);
+	for (UObject * Obj : ChildObjects)
 	{
-		UGameplayAbility * Ability = NULL;
-		if (BindInfo.GameplayAbilityClass)
+		UAttributeSet * Set = Cast<UAttributeSet>(Obj);
+		if (Set)  
 		{
-			Ability = BindInfo.GameplayAbilityClass->GetDefaultObject<UGameplayAbility>();
-			if (Ability->AllowInstancing())
-			{
-				Ability = ConstructObject<UGameplayAbility>(BindInfo.GameplayAbilityClass, OwnerActor);
-				InstancedAbilities.Add(Ability);
-			}
-			AllAbilities.Add(Ability);
+			UObject * AT = Set->GetArchetype();
+
+			//ensure(SpawnedAttributes.Contains(Set) == false);
+			//ensure(Set->IsDefaultSubobject());
+			SpawnedAttributes.Add(Set);
 		}
 	}
 }
 
-/** This is a convenience/helper function to bind input directly to abilities. Some games may need to do things differently, if for example you want weapons to fire abilities and the pawn to swap weapons. */
-void UAttributeComponent::BindInputComponentToAbilities(UInputComponent *InputComponent, UGameplayAbilitySet *AbilitySet)
+UGameplayAbility * UAttributeComponent::CreateNewInstanceOfAbility(UGameplayAbility *Ability)
 {
-	check(AbilitySet->Abilities.Num() == AllAbilities.Num());
+	check(Ability);
+	check(Ability->HasAllFlags(RF_ClassDefaultObject));
 
 	AActor * OwnerActor = GetOwner();
 	check(OwnerActor);
 
-	FGameplayAbilityActorInfo	ActorInfo;
-	ActorInfo.InitFromActor(OwnerActor);
+	UGameplayAbility * AbilityInstance = ConstructObject<UGameplayAbility>(Ability->GetClass(), OwnerActor);
+	check(AbilityInstance);
 
-	for (int32 idx = 0; idx < AbilitySet->Abilities.Num(); ++idx)
-	{
-		FGameplayAbilityBindInfo & BindInfo = AbilitySet->Abilities[idx];
-		UGameplayAbility * Ability = AllAbilities[idx];
-		
-		check(Ability->GetClass() == AbilitySet->Abilities[idx].GameplayAbilityClass);
-		
-		for (const FGameplayAbilityBindInfoCommandIDPair &CommandPair : BindInfo.CommandBindings)
-		{
-			// Pressed event
-			{
-				FInputActionBinding AB(FName(*CommandPair.CommandString), IE_Pressed);
-				AB.ActionDelegate.GetDelegateForManualSet().BindUObject(Ability, &UGameplayAbility::InputPressed, CommandPair.InputID, ActorInfo);
-				InputComponent->AddActionBinding(AB);
-			}
-
-			// Released Event
-			{
-				FInputActionBinding AB(FName(*CommandPair.CommandString), IE_Released);
-				AB.ActionDelegate.GetDelegateForManualSet().BindUObject(Ability, &UGameplayAbility::InputReleased, CommandPair.InputID, ActorInfo);
-				InputComponent->AddActionBinding(AB);
-			}
-		}
-	}
-}
-
-void UAttributeComponent::MontageBranchPoint_AbilityDecisionStop()
-{
-	if (AnimatingAbility)
+	// Add it to one of our instance lists so that it doesn't GC.
+	if (AbilityInstance->GetReplicationPolicy() != EGameplayAbilityReplicationPolicy::ReplicateNone)
 	{
-		AActor * OwnerActor = GetOwner();
-		check(OwnerActor);
-
-		FGameplayAbilityActorInfo	ActorInfo;
-		ActorInfo.InitFromActor(OwnerActor);
-
-		AnimatingAbility->MontageBranchPoint_AbilityDecisionStop(ActorInfo);
+		ReplicatedInstancedAbilities.Add(AbilityInstance);
 	}
-
-}
-
-void UAttributeComponent::MontageBranchPoint_AbilityDecisionStart()
-{
-	if (AnimatingAbility)
+	else
 	{
-		AActor * OwnerActor = GetOwner();
-		check(OwnerActor);
-
-		FGameplayAbilityActorInfo	ActorInfo;
-		ActorInfo.InitFromActor(OwnerActor);
-
-		AnimatingAbility->MontageBranchPoint_AbilityDecisionStart(ActorInfo);
+		NonReplicatedInstancedAbilities.Add(AbilityInstance);
 	}
+	
+	return AbilityInstance;
 }
 
 bool UAttributeComponent::ActivateAbility(int32 AbilityIdx)
@@ -115,17 +68,72 @@ bool UAttributeComponent::ActivateAbility(int32 AbilityIdx)
 	FGameplayAbilityActorInfo	ActorInfo;
 	ActorInfo.InitFromActor(OwnerActor);
 
-	if (AllAbilities.IsValidIndex(AbilityIdx))
+	if (ActivatableAbilities.IsValidIndex(AbilityIdx))
 	{
 		// Fixme
-		AllAbilities[AbilityIdx]->InputPressed(0, ActorInfo);
+		ActivatableAbilities[AbilityIdx]->InputPressed(0, ActorInfo);
 	}
 
 	return false;
 }
 
-void UAttributeComponent::ServerTryActivateAbility_Implementation(class UGameplayAbility * AbilityToActivate)
+void UAttributeComponent::CancelAbilitiesWithTags(const FGameplayTagContainer Tags, const FGameplayAbilityActorInfo & ActorInfo, UGameplayAbility * Ignore)
+{
+	/**
+	 *	Right now we are cancelling all activatable abilites that match Tags. This includes abilities that might not have been activated in the first place!
+	 *	For instanced-per-actor abilities this is fine. They could check if they were activated/still activating.
+	 *	For non-instanced abilities it is ambigous. We have no way to know 'how many' non-instanced abilities are in flight.
+	 *	Likewise for instanced-per-execution abilities. Though they are present in Replicated/NonReplicatedInstancedAbilities list
+	 */
+
+	struct local
+	{
+		static void CancelAbilitiesWithTags(const FGameplayTagContainer Tags, TArray<UGameplayAbility*> Abilities, const FGameplayAbilityActorInfo ActorInfo, UGameplayAbility * Ignore)
+		{
+			for (int32 idx=0; idx < Abilities.Num(); ++idx)
+			{
+				UGameplayAbility *Ability = Abilities[idx];
+				if (Ability && (Ability != Ignore) && Ability->AbilityTags.MatchesAny(Tags, false))
+				{
+					Ability->CancelAbility(ActorInfo);
+					if (!Ability->HasAnyFlags(RF_ClassDefaultObject) && Ability->GetInstancingPolicy() == EGameplayAbilityInstancingPolicy::InstancedPerExecution)
+					{
+						Ability->MarkPendingKill();
+						Abilities.RemoveAtSwap(idx);
+						idx--;
+					}
+				}
+			}
+		}
+	};
+
+	//local::CancelAbilitiesWithTags(Tags, ReplicatedInstancedAbilities, ActorInfo, Ignore);
+	//local::CancelAbilitiesWithTags(Tags, NonReplicatedInstancedAbilities, ActorInfo, Ignore);
+
+	local::CancelAbilitiesWithTags(Tags, ActivatableAbilities, ActorInfo, Ignore);
+}
+
+#if !(UE_BUILD_SHIPPING || UE_BUILD_TEST)
+int32 DenyClientActivation = 0;
+static FAutoConsoleVariableRef CVarDenyClientActivation(
+	TEXT("SkillSystem.DenyClientActivations"),
+	DenyClientActivation,
+	TEXT("Make server deny the next X ability activations from clients. For testing misprediction."),
+	ECVF_Default
+	);
+#endif
+
+void UAttributeComponent::ServerTryActivateAbility_Implementation(class UGameplayAbility * AbilityToActivate, int32 PredictionKey)
 {
+#if !(UE_BUILD_SHIPPING || UE_BUILD_TEST)
+	if (DenyClientActivation > 0)
+	{
+		DenyClientActivation--;
+		ClientActivateAbilityFailed(AbilityToActivate, PredictionKey);
+		return;
+	}
+#endif
+
 	ensure(AbilityToActivate);
 
 	AActor * OwnerActor = GetOwner();
@@ -133,29 +141,65 @@ void UAttributeComponent::ServerTryActivateAbility_Implementation(class UGamepla
 
 	FGameplayAbilityActorInfo	ActorInfo;
 	ActorInfo.InitFromActor(OwnerActor);
-	
-	if (AbilityToActivate->CanActivateAbility(ActorInfo))
+	ActorInfo.SetPredictionKey(PredictionKey);	// Set the prediction key in case anyone else needs it. But we are still Authority and not Predicting.
+
+	UGameplayAbility *InstancedAbility = NULL;
+
+	if (AbilityToActivate->TryActivateAbility(ActorInfo, &InstancedAbility))
 	{
-		AbilityToActivate->ActivateAbility(ActorInfo);
-		ClientActivateAbilitySucceed(AbilityToActivate);
+		if (InstancedAbility && InstancedAbility->GetReplicationPolicy() != EGameplayAbilityReplicationPolicy::ReplicateNone)
+		{
+			InstancedAbility->ClientActivateAbilitySucceed(OwnerActor);
+		}
+		else
+		{
+			ClientActivateAbilitySucceed(AbilityToActivate, PredictionKey);
+		}
+
+		// Update our ReplicatedPredictionKey. When the client gets value, he will know his state (actor+all components/subobjects) are up to do date and he can
+		// remove any necessary predictive work.
+
+		if (PredictionKey > 0)
+		{
+			ensure(PredictionKey > ReplicatedPredictionKey);
+			ReplicatedPredictionKey = PredictionKey;
+		}
 	}
 	else
 	{
-		ClientActivateAbilityFailed(AbilityToActivate);
+		ClientActivateAbilityFailed(AbilityToActivate, PredictionKey);
 	}
 }
 
-bool UAttributeComponent::ServerTryActivateAbility_Validate(class UGameplayAbility * AbilityToActivate)
+bool UAttributeComponent::ServerTryActivateAbility_Validate(class UGameplayAbility * AbilityToActivate, int32 PredictionKey)
 {
 	return true;
 }
 
-void UAttributeComponent::ClientActivateAbilityFailed_Implementation(class UGameplayAbility * AbilityToActivate)
+void UAttributeComponent::ClientActivateAbilityFailed_Implementation(class UGameplayAbility * AbilityToActivate, int32 PredictionKey)
 {
+	if (PredictionKey > 0)
+	{
+		for (int32 idx = 0; idx < PredictionDelegates.Num(); ++idx)
+		{
+			int32 Key = PredictionDelegates[idx].Key;
+			if (Key == PredictionKey)
+			{
+				SKILL_LOG(Warning, TEXT("Failed ActivateAbility, clearing prediction data %d"), PredictionKey);
 
+				PredictionDelegates[idx].Value.Broadcast();
+				PredictionDelegates.RemoveAt(idx);
+				break;
+			}
+			else if(Key >= PredictionKey)
+			{
+				break;
+			}
+		}
+	}
 }
 
-void UAttributeComponent::ClientActivateAbilitySucceed_Implementation(class UGameplayAbility * AbilityToActivate)
+void UAttributeComponent::ClientActivateAbilitySucceed_Implementation(class UGameplayAbility * AbilityToActivate, int32 PredictionKey)
 {
 	ensure(AbilityToActivate);
 
@@ -164,10 +208,39 @@ void UAttributeComponent::ClientActivateAbilitySucceed_Implementation(class UGam
 
 	FGameplayAbilityActorInfo	ActorInfo;
 	ActorInfo.InitFromActor(OwnerActor);
+	ActorInfo.SetPredictionKey(PredictionKey);
+	ActorInfo.SetActivationConfirmed();
 
-	AbilityToActivate->ActivateAbility(ActorInfo);
+	AbilityToActivate->CallActivateAbility(ActorInfo);
 }
 
+void UAttributeComponent::MontageBranchPoint_AbilityDecisionStop()
+{
+	if (AnimatingAbility)
+	{
+		AActor * OwnerActor = GetOwner();
+		check(OwnerActor);
+
+		FGameplayAbilityActorInfo	ActorInfo;
+		ActorInfo.InitFromActor(OwnerActor);
+
+		AnimatingAbility->MontageBranchPoint_AbilityDecisionStop(ActorInfo);
+	}
+}
+
+void UAttributeComponent::MontageBranchPoint_AbilityDecisionStart()
+{
+	if (AnimatingAbility)
+	{
+		AActor * OwnerActor = GetOwner();
+		check(OwnerActor);
+
+		FGameplayAbilityActorInfo	ActorInfo;
+		ActorInfo.InitFromActor(OwnerActor);
+
+		AnimatingAbility->MontageBranchPoint_AbilityDecisionStart(ActorInfo);
+	}
+}
 
 #undef LOCTEXT_NAMESPACE
 
diff --git a/Engine/Source/Runtime/SkillSystem/Private/BlueprintWaitMovementModeChangeTaskProxy.cpp b/Engine/Source/Runtime/SkillSystem/Private/BlueprintWaitMovementModeChangeTaskProxy.cpp
new file mode 100644
--- /dev/null
+++ b/Engine/Source/Runtime/SkillSystem/Private/BlueprintWaitMovementModeChangeTaskProxy.cpp
@@ -0,0 +1,24 @@
+
+#include "SkillSystemModulePrivatePCH.h"
+
+UBlueprintWaitMovementModeChangeTaskProxy::UBlueprintWaitMovementModeChangeTaskProxy(const class FPostConstructInitializeProperties& PCIP)
+	: Super(PCIP)
+{
+	RequiredMode = MOVE_None;
+}
+
+void UBlueprintWaitMovementModeChangeTaskProxy::OnMovementModeChange(ACharacter * Character, EMovementMode PrevMovementMode, uint8 PreviousCustomMode)
+{
+	if (Character)
+	{
+		if (UCharacterMovementComponent *MoveComp = Cast<UCharacterMovementComponent>(Character->GetMovementComponent()))
+		{
+			if (RequiredMode == MOVE_None || MoveComp->MovementMode == RequiredMode)
+			{
+				Character->MovementModeChangedDelegate.RemoveDynamic(this, &UBlueprintWaitMovementModeChangeTaskProxy::OnMovementModeChange);
+				OnChange.Broadcast(MoveComp->MovementMode);
+				return;
+			}
+		}
+	}
+}
\ No newline at end of file
diff --git a/Engine/Source/Runtime/SkillSystem/Private/GameplayAbility.cpp b/Engine/Source/Runtime/SkillSystem/Private/GameplayAbility.cpp
deleted file mode 100644
diff --git a/Engine/Source/Runtime/SkillSystem/Private/GameplayAbilitySet.cpp b/Engine/Source/Runtime/SkillSystem/Private/GameplayAbilitySet.cpp
--- a/Engine/Source/Runtime/SkillSystem/Private/GameplayAbilitySet.cpp
+++ b/Engine/Source/Runtime/SkillSystem/Private/GameplayAbilitySet.cpp
@@ -1,4 +1,4 @@
-// Copyright 1998-2013 Epic Games, Inc. All Rights Reserved.
+// Copyright 1998-2014 Epic Games, Inc. All Rights Reserved.
 
 #include "SkillSystemModulePrivatePCH.h"
 
@@ -8,6 +8,50 @@
 //
 // --------------------------------------------------------------------------------------------------------------------------------------------------------
 
+void InputPressed(int32 AbilityIdx, int32 InputID, TWeakObjectPtr<AActor> ActorOwner)
+{
+	FGameplayAbilityActorInfo	ActorInfo;
+	ActorInfo.InitFromActor(ActorOwner.Get());
+
+	UAttributeComponent * AttributeComponent = ActorInfo.AttributeComponent.Get();
+	if (AttributeComponent)
+	{
+		if (AttributeComponent->ActivatableAbilities.IsValidIndex(AbilityIdx))
+		{
+			UGameplayAbility * Ability = AttributeComponent->ActivatableAbilities[AbilityIdx];
+			if (Ability)
+			{
+				Ability->InputPressed(InputID, ActorInfo);
+			}
+		}
+	}
+}
+
+void InputReleased(int32 AbilityIdx, int32 InputID, TWeakObjectPtr<AActor> ActorOwner)
+{
+	FGameplayAbilityActorInfo	ActorInfo;
+	ActorInfo.InitFromActor(ActorOwner.Get());
+
+	UAttributeComponent * AttributeComponent = ActorInfo.AttributeComponent.Get();
+	if (AttributeComponent)
+	{
+		if (AttributeComponent->ActivatableAbilities.IsValidIndex(AbilityIdx))
+		{
+			UGameplayAbility * Ability = AttributeComponent->ActivatableAbilities[AbilityIdx];
+			if (Ability)
+			{
+				Ability->InputReleased(InputID, ActorInfo);
+			}
+		}
+	}
+}
+
+// --------------------------------------------------------------------------------------------------------------------------------------------------------
+//
+//	UGameplayAbilitySet
+//
+// --------------------------------------------------------------------------------------------------------------------------------------------------------
+
 
 UGameplayAbilitySet::UGameplayAbilitySet(const class FPostConstructInitializeProperties& PCIP)
 	: Super(PCIP)
@@ -15,3 +59,58 @@ UGameplayAbilitySet::UGameplayAbilitySet(const class FPostConstructInitializePro
 
 }
 
+void UGameplayAbilitySet::InitializeAbilities(UAttributeComponent *AttributeComponent) const
+{
+	AActor * OwnerActor = AttributeComponent->GetOwner();
+	check(OwnerActor);
+
+	FGameplayAbilityActorInfo	ActorInfo;
+	ActorInfo.InitFromActor(OwnerActor);
+
+	for (const FGameplayAbilityBindInfo &BindInfo : this->Abilities)
+	{
+		UGameplayAbility * Ability = BindInfo.GameplayAbilityClass ? BindInfo.GameplayAbilityClass->GetDefaultObject<UGameplayAbility>() : BindInfo.GameplayAbilityInstance;
+		if (Ability)
+		{
+			if (Ability->GetInstancingPolicy() == EGameplayAbilityInstancingPolicy::InstancedPerActor)
+			{
+				Ability = AttributeComponent->CreateNewInstanceOfAbility(Ability);
+			}
+
+			AttributeComponent->ActivatableAbilities.Add(Ability);
+		}
+	}
+}
+
+void UGameplayAbilitySet::BindInputComponentToAbilities(UInputComponent *InputComponent, UAttributeComponent *AttributeComponent) const
+{
+	AActor * OwnerActor = AttributeComponent->GetOwner();
+	check(OwnerActor);
+
+	FGameplayAbilityActorInfo	ActorInfo;
+	ActorInfo.InitFromActor(OwnerActor);
+
+	for (int32 idx = 0; idx < this->Abilities.Num(); ++idx)
+	{
+		const FGameplayAbilityBindInfo & BindInfo = this->Abilities[idx];
+		UGameplayAbility * Ability = AttributeComponent->ActivatableAbilities[idx];	
+		
+		for (const FGameplayAbilityBindInfoCommandIDPair &CommandPair : BindInfo.CommandBindings)
+		{
+			// Pressed event
+			{
+				FInputActionBinding AB(FName(*CommandPair.CommandString), IE_Pressed);
+				AB.ActionDelegate.GetDelegateForManualSet().BindStatic(&InputPressed, idx, CommandPair.InputID, TWeakObjectPtr<AActor>(OwnerActor));
+				InputComponent->AddActionBinding(AB);
+			}
+
+			// Released event
+			{
+				FInputActionBinding AB(FName(*CommandPair.CommandString), IE_Released);
+				AB.ActionDelegate.GetDelegateForManualSet().BindStatic(&InputReleased, idx, CommandPair.InputID, TWeakObjectPtr<AActor>(OwnerActor));
+				InputComponent->AddActionBinding(AB);
+			}
+		}
+	}
+}
+
diff --git a/Engine/Source/Runtime/SkillSystem/Private/GameplayAbility_Instanced.cpp b/Engine/Source/Runtime/SkillSystem/Private/GameplayAbility_Instanced.cpp
deleted file mode 100644
diff --git a/Engine/Source/Runtime/SkillSystem/Private/GameplayAbility_Montage.cpp b/Engine/Source/Runtime/SkillSystem/Private/GameplayAbility_Montage.cpp
deleted file mode 100644
diff --git a/Engine/Source/Runtime/SkillSystem/Private/GameplayEffect.cpp b/Engine/Source/Runtime/SkillSystem/Private/GameplayEffect.cpp
--- a/Engine/Source/Runtime/SkillSystem/Private/GameplayEffect.cpp
+++ b/Engine/Source/Runtime/SkillSystem/Private/GameplayEffect.cpp
@@ -1444,12 +1444,19 @@ void FActiveGameplayEffectsContainer::OnPropertyAggregatorDirty(FAggregator* Agg
 	Owner->SetNumericAttribute(Attribute, NewPropertyValue);
 }
 
-FActiveGameplayEffect & FActiveGameplayEffectsContainer::CreateNewActiveGameplayEffect(const FGameplayEffectSpec &Spec)
+FActiveGameplayEffect & FActiveGameplayEffectsContainer::CreateNewActiveGameplayEffect(const FGameplayEffectSpec &Spec, int32 InPredictionKey)
 {
 	LastAssignedHandle = LastAssignedHandle.GetNextHandle();
-	FActiveGameplayEffect & NewEffect = *new (GameplayEffects)FActiveGameplayEffect(LastAssignedHandle, Spec, GetWorldTime(), GetGameStateTime());
+	FActiveGameplayEffect & NewEffect = *new (GameplayEffects)FActiveGameplayEffect(LastAssignedHandle, Spec, GetWorldTime(), GetGameStateTime(), InPredictionKey);
 	
-	MarkItemDirty(NewEffect);
+	if (InPredictionKey == 0 || IsNetAuthority())	// Clients predicting a GameplayEffect must not call MarkItemDirty
+	{
+		MarkItemDirty(NewEffect);
+	}
+	else
+	{
+		Owner->GetPredictionKeyDelegate(InPredictionKey).AddUObject(Owner, &UAttributeComponent::RemoveActiveGameplayEffect_NoReturn, NewEffect.Handle);
+	}
 	return NewEffect;
 }
 
@@ -1689,6 +1696,38 @@ bool FActiveGameplayEffectsContainer::HasAnyTags(FGameplayTagContainer &Tags)
 	return false;
 }
 
+bool FActiveGameplayEffectsContainer::HasAllTags(FGameplayTagContainer &Tags)
+{
+	SCOPE_CYCLE_COUNTER(STAT_GameplayEffectsHasAllTags);
+
+	for (auto It = Tags.CreateConstIterator(); It;  ++It)
+	{
+		bool Found = false;
+		for (FActiveGameplayEffect &ActiveEffect : GameplayEffects)
+		{
+			// Fixme: check stacking rules!
+
+			/**
+			 *	The desired behavior here is:
+			 *		"Do we have tag a.b.c and the GameplayEffect has a.b.c.d -> match!"
+			 *		"Do we have tag a.b.c.d. and the GameplayEffect has a.b.c -> no match!"
+			 */
+
+			if (ActiveEffect.Spec.Def->OwnedTagsContainer.HasTag(*It, EGameplayTagMatchType::IncludeParentTags, EGameplayTagMatchType::Explicit))
+			{
+				Found = true;
+				break;
+			}
+		}
+		if (!Found)
+		{
+			return false;
+		}
+	}
+
+	return true;
+}
+
 bool FActiveGameplayEffectsContainer::CanApplyAttributeModifiers(const UGameplayEffect *GameplayEffect, float Level, AActor *Instigator)
 {
 	SCOPE_CYCLE_COUNTER(STAT_GameplayEffectsCanApplyAttributeModifiers);
diff --git a/Engine/Source/Runtime/SkillSystem/Private/SkillSystem.cpp b/Engine/Source/Runtime/SkillSystem/Private/SkillSystem.cpp
--- a/Engine/Source/Runtime/SkillSystem/Private/SkillSystem.cpp
+++ b/Engine/Source/Runtime/SkillSystem/Private/SkillSystem.cpp
@@ -2,6 +2,7 @@
 
 #include "SkillSystemModulePrivatePCH.h"
 
+DEFINE_STAT(STAT_GameplayEffectsHasAllTags);
 DEFINE_STAT(STAT_GameplayEffectsHasAnyTag);
 DEFINE_STAT(STAT_GameplayEffectsTick);
 DEFINE_STAT(STAT_GameplayEffectsCanApplyAttributeModifiers);
diff --git a/Engine/Source/Runtime/SkillSystem/Private/SkillSystemBlueprintLibrary.cpp b/Engine/Source/Runtime/SkillSystem/Private/SkillSystemBlueprintLibrary.cpp
--- a/Engine/Source/Runtime/SkillSystem/Private/SkillSystemBlueprintLibrary.cpp
+++ b/Engine/Source/Runtime/SkillSystem/Private/SkillSystemBlueprintLibrary.cpp
@@ -8,6 +8,7 @@ USkillSystemBlueprintLibrary::USkillSystemBlueprintLibrary(const class FPostCons
 {
 }
 
+// Fixme: remove this function!!!
 void USkillSystemBlueprintLibrary::StaticFuncPlayMontageAndWait(AActor* WorldContextObject, UAnimMontage * Montage, FLatentActionInfo LatentInfo)
 {
 	check(WorldContextObject);
@@ -18,9 +19,20 @@ void USkillSystemBlueprintLibrary::StaticFuncPlayMontageAndWait(AActor* WorldCon
 		if (LatentActionManager.FindExistingAction<FPlayMontageAndWaitAction>(LatentInfo.CallbackTarget, LatentInfo.UUID) == NULL)
 		{
 			FGameplayAbilityActorInfo	ActorInfo;
-			ActorInfo.InitFromActor(WorldContextObject);
 
-			LatentActionManager.AddNewAction(LatentInfo.CallbackTarget, LatentInfo.UUID, new FPlayMontageAndWaitAction(ActorInfo, Montage, LatentInfo));
+			AActor * ActorOwner = NULL;
+			UObject * TestObj = WorldContextObject;
+			while(ActorOwner == NULL && TestObj != NULL)
+			{
+				ActorOwner = Cast<AActor>(TestObj);
+				TestObj = TestObj->GetOuter();
+			}
+
+			if (ActorOwner)
+			{
+				ActorInfo.InitFromActor(ActorOwner);  
+				LatentActionManager.AddNewAction(LatentInfo.CallbackTarget, LatentInfo.UUID, new FPlayMontageAndWaitAction(ActorInfo, Montage, LatentInfo));
+			}
 		}
 	}
 }
@@ -49,6 +61,46 @@ UBlueprintPlayMontageAndWaitTaskProxy* USkillSystemBlueprintLibrary::CreatePlayM
 		return MyObj;
 	}
 
+	return NULL;	
+}
+
+
+UBlueprintWaitMovementModeChangeTaskProxy* USkillSystemBlueprintLibrary::CreateWaitMovementModeChange(class UObject* WorldContextObject, EMovementMode NewMode)
+{
+	check(WorldContextObject);
+
+	UGameplayAbility_Instanced * Ability = Cast<UGameplayAbility_Instanced>(WorldContextObject);
+	if (Ability)
+	{
+		AActor * ActorOwner = Cast<AActor>(Ability->GetOuter());
+
+		UBlueprintWaitMovementModeChangeTaskProxy * MyObj = NULL;
+		MyObj = NewObject<UBlueprintWaitMovementModeChangeTaskProxy>();
+		MyObj->RequiredMode = NewMode;
+
+		FGameplayAbilityActorInfo	ActorInfo;
+		ActorInfo.InitFromActor(ActorOwner);
+
+		ACharacter * Character = Cast<ACharacter>(ActorInfo.Actor.Get());
+		if (Character)
+		{
+			/*
+				FScriptDelegate Delegate;
+				Delegate.BindUFunction(this, "OnActorBump");
+				MyPawn->OnActorHit.AddUnique(Delegate);
+			*/
+
+			//FMovementModeChangedSignature::FDelegate::CreateUObject(MyObj, &UBlueprintWaitMovementModeChangeTaskProxy::OnMovementModeChange);
+			//MyDel.BindUObject( 
+			//Character->MovementModeChangedDelegate.Add(MyObj, &UBlueprintWaitMovementModeChangeTaskProxy::OnMovementModeChange); // //FMovementModeChangedSignature::FDelegate::CreateUObject
+
+			Character->MovementModeChangedDelegate.AddDynamic(MyObj, &UBlueprintWaitMovementModeChangeTaskProxy::OnMovementModeChange);
+
+			//Character->MovementModeChangedDelegate.AddUObject(MyObj, &UBlueprintWaitMovementModeChangeTaskProxy::OnMovementModeChange);
+		}
+
+		return MyObj;
+	}
+
 	return NULL;
-	
 }
\ No newline at end of file
diff --git a/Engine/Source/Runtime/SkillSystem/Private/SkillSystemDebugHUD.cpp b/Engine/Source/Runtime/SkillSystem/Private/SkillSystemDebugHUD.cpp
new file mode 100644
--- /dev/null
+++ b/Engine/Source/Runtime/SkillSystem/Private/SkillSystemDebugHUD.cpp
@@ -0,0 +1,161 @@
+// Copyright 1998-2014 Epic Games, Inc. All Rights Reserved.
+
+#include "SkillSystemModulePrivatePCH.h"
+#include "SkillSystemDebugHUD.h"
+#include "DebugRenderSceneProxy.h"
+
+ASkillSystemDebugHUD::ASkillSystemDebugHUD(const class FPostConstructInitializeProperties& PCIP) : Super(PCIP)
+{
+
+}
+
+void ASkillSystemDebugHUD::DrawHUD()
+{
+	Super::DrawHUD();
+}
+
+void ASkillSystemDebugHUD::DrawWithBackground(UFont* InFont, const FString& Text, const FColor& TextColor, EAlignHorizontal::Type HAlign, float& OffsetX, EAlignVertical::Type VAlign, float& OffsetY, float Alpha)
+{
+	float SizeX, SizeY;
+	Canvas->StrLen(InFont, Text, SizeX, SizeY);
+
+	const float PosX = (HAlign == EAlignHorizontal::Center) ? OffsetX + (Canvas->ClipX - SizeX) * 0.5f :
+		(HAlign == EAlignHorizontal::Left) ? Canvas->OrgX + OffsetX :
+		Canvas->ClipX - SizeX - OffsetX;
+
+	const float PosY = (VAlign == EAlignVertical::Center) ? OffsetY + (Canvas->ClipY - SizeY) * 0.5f :
+		(VAlign == EAlignVertical::Top) ? Canvas->OrgY + OffsetY :
+		Canvas->ClipY - SizeY - OffsetY;
+
+	const float BoxPadding = 5.0f;
+
+	const float X = PosX - BoxPadding;
+	const float Y = PosY - BoxPadding;
+	const float Z = 0.1f;
+	FCanvasTileItem TileItem(FVector2D(X, Y), FVector2D(SizeX + BoxPadding * 2.0f, SizeY + BoxPadding * 2.0f), FLinearColor(0.75f, 0.75f, 0.75f, Alpha));
+	Canvas->DrawItem(TileItem);
+
+	FLinearColor TextCol(TextColor);
+	TextCol.A = Alpha;
+	FCanvasTextItem TextItem(FVector2D(PosX, PosY), FText::FromString(Text), GEngine->GetSmallFont(), TextCol);
+	Canvas->DrawItem(TextItem);
+
+	OffsetY += 25;
+}
+
+void ASkillSystemDebugHUD::DrawDebugHUD(UCanvas* InCanvas, APlayerController* PC)
+{
+	Canvas = InCanvas;
+	if (!Canvas)
+	{
+		return;
+	}
+
+	UPlayer * LocalPlayer = GetWorld()->GetFirstLocalPlayerFromController();
+	PC = LocalPlayer->PlayerController;	
+
+	if (!PC)
+	{
+		return;
+	}
+
+	if (PC->GetPawn())
+	{
+		UAttributeComponent * AttributeComponent = PC->GetPawn()->FindComponentByClass<UAttributeComponent>();
+		if (AttributeComponent)
+		{
+			DrawDebugAttributeComponent(AttributeComponent);
+		}
+	}
+}
+
+void ASkillSystemDebugHUD::DrawDebugAttributeComponent(UAttributeComponent *Component)
+{
+	UWorld *World = GetWorld();
+	float GameWorldTime = World->GetTimeSeconds();
+
+	UFont* Font = GEngine->GetSmallFont();
+	FColor Color(38, 128, 0);
+	float X = 20.f;
+	float Y = 20.f;
+
+	FString String = FString::Printf(TEXT("%s (%d)"), *Component->GetPathName(), Component->IsDefaultSubobject());
+	DrawWithBackground(Font, String, Color, EAlignHorizontal::Left, X, EAlignVertical::Top, Y);
+
+	
+	String = FString::Printf(TEXT("%s == %s"), *Component->GetArchetype()->GetPathName(), *Component->GetClass()->GetDefaultObject()->GetPathName());
+	DrawWithBackground(Font, String, Color, EAlignHorizontal::Left, X, EAlignVertical::Top, Y);
+
+
+	for (UAttributeSet * Set : Component->SpawnedAttributes)
+	{
+		if (!Set)
+			continue;
+		check(Set);
+
+		// Draw Attribute Set
+		DrawWithBackground(Font, FString::Printf(TEXT("%s (%d)"), *Set->GetName(), Set->IsDefaultSubobject()), Color, EAlignHorizontal::Left, X, EAlignVertical::Top, Y);
+
+		String = FString::Printf(TEXT("%s == %s"), *Set->GetArchetype()->GetPathName(), *Set->GetClass()->GetDefaultObject()->GetPathName());
+		DrawWithBackground(Font, String, Color, EAlignHorizontal::Left, X, EAlignVertical::Top, Y);
+
+		for (TFieldIterator<UProperty> PropertyIt(Set->GetClass(), EFieldIteratorFlags::IncludeSuper); PropertyIt; ++PropertyIt)
+		{
+			UProperty *Prop = *PropertyIt;			
+
+			FString ValueString;
+			void *PropertyValue = Prop->ContainerPtrToValuePtr<void>(Set);
+			Prop->ExportTextItem(ValueString, PropertyValue, NULL, NULL, 0);
+
+			String = FString::Printf(TEXT("%s: %s"), *Prop->GetDisplayNameText().ToString(), *ValueString);
+			DrawWithBackground(Font, String, Color, EAlignHorizontal::Left, X, EAlignVertical::Top, Y);
+		}
+
+		Y+= 25;
+		// Draw Active GameplayEffect
+		for (FActiveGameplayEffect &Effect : Component->ActiveGameplayEffects.GameplayEffects)
+		{
+			String = FString::Printf(TEXT("%s. [%d] %.2f"), *Effect.Spec.ToSimpleString(), Effect.PredictionKey, Effect.GetTimeRemaining(GameWorldTime));
+			DrawWithBackground(Font, String, Color, EAlignHorizontal::Left, X, EAlignVertical::Top, Y);	
+		}
+	}
+	
+}
+
+#if !(UE_BUILD_SHIPPING || UE_BUILD_TEST)
+static void	ToggleDebugHUD(const TArray<FString>& Args, UWorld* InWorld)
+{
+
+	if (!InWorld)
+		return;
+
+	ASkillSystemDebugHUD *HUD = NULL;
+
+	for (FActorIterator It(InWorld); It; ++It)
+	{
+		AActor *Actor = *It;
+		HUD = Cast<ASkillSystemDebugHUD>(Actor);
+		if (HUD)
+			break;
+	}
+
+	if (!HUD)
+	{
+		HUD = InWorld->SpawnActor<ASkillSystemDebugHUD>();
+		
+		FDebugDrawDelegate DrawDebugDelegate = FDebugDrawDelegate::CreateUObject(HUD, &ASkillSystemDebugHUD::DrawDebugHUD);
+		UDebugDrawService::Register(TEXT("GameplayDebug"), DrawDebugDelegate);
+	}
+	else
+	{
+		FDebugDrawDelegate DrawDebugDelegate = FDebugDrawDelegate::CreateUObject(HUD, &ASkillSystemDebugHUD::DrawDebugHUD);
+		UDebugDrawService::Unregister(DrawDebugDelegate);
+	}
+}
+
+FAutoConsoleCommandWithWorldAndArgs TestObjectRefSerializeCommand(
+	TEXT("SkillSystem.ToggleDebugHUD"),
+	TEXT("ToggleDebugHUD Drawing"),
+	FConsoleCommandWithWorldAndArgsDelegate::CreateStatic(ToggleDebugHUD)
+	);
+#endif
diff --git a/Engine/Source/Runtime/SkillSystem/Private/SkillSystemModulePrivatePCH.h b/Engine/Source/Runtime/SkillSystem/Private/SkillSystemModulePrivatePCH.h
--- a/Engine/Source/Runtime/SkillSystem/Private/SkillSystemModulePrivatePCH.h
+++ b/Engine/Source/Runtime/SkillSystem/Private/SkillSystemModulePrivatePCH.h
@@ -1,4 +1,4 @@
-// Copyright 1998-2013 Epic Games, Inc. All Rights Reserved.
+// Copyright 1998-2014 Epic Games, Inc. All Rights Reserved.
 
 #pragma once
 
@@ -13,6 +13,9 @@
 #include "Core.h"
 #include "Engine.h"
 #include "Engine/DataTable.h"
+
+#include "SkillSystemTypes.h"
+
 #include "SkillSystemClasses.h"
 
 #include "SkillSystem.h"
diff --git a/Engine/Source/Runtime/SkillSystem/Public/SkillSystem.h b/Engine/Source/Runtime/SkillSystem/Public/SkillSystem.h
--- a/Engine/Source/Runtime/SkillSystem/Public/SkillSystem.h
+++ b/Engine/Source/Runtime/SkillSystem/Public/SkillSystem.h
@@ -9,6 +9,7 @@
 
 DECLARE_STATS_GROUP(TEXT("SkillSystem"), STATGROUP_SkillSystem, STATCAT_Advanced);
 
+DECLARE_CYCLE_STAT_EXTERN(TEXT("GameplayEffectsHasAllTagsTime"), STAT_GameplayEffectsHasAllTags, STATGROUP_SkillSystem, );
 DECLARE_CYCLE_STAT_EXTERN(TEXT("GameplayEffectsHasAnyTagTime"), STAT_GameplayEffectsHasAnyTag, STATGROUP_SkillSystem, );
 DECLARE_CYCLE_STAT_EXTERN(TEXT("GameplayEffectsTick"), STAT_GameplayEffectsTick, STATGROUP_SkillSystem, );
 DECLARE_CYCLE_STAT_EXTERN(TEXT("CanApplyAttributeModifiers"), STAT_GameplayEffectsCanApplyAttributeModifiers, STATGROUP_SkillSystem, );
diff --git a/Engine/Source/Runtime/SkillSystem/Public/SkillSystemLog.h b/Engine/Source/Runtime/SkillSystem/Public/SkillSystemLog.h
--- a/Engine/Source/Runtime/SkillSystem/Public/SkillSystemLog.h
+++ b/Engine/Source/Runtime/SkillSystem/Public/SkillSystemLog.h
@@ -9,7 +9,7 @@
 //
 
 
-SKILLSYSTEM_API DECLARE_LOG_CATEGORY_EXTERN(LogSkillSystem, Log, All);
+SKILLSYSTEM_API DECLARE_LOG_CATEGORY_EXTERN(LogSkillSystem, Warning, All);
 
 #define SKILL_LOG(Verbosity, Format, ...) \
 { \
diff --git a/Engine/Source/Runtime/SkillSystem/Public/SkillSystemTypes.h b/Engine/Source/Runtime/SkillSystem/Public/SkillSystemTypes.h
new file mode 100644
--- /dev/null
+++ b/Engine/Source/Runtime/SkillSystem/Public/SkillSystemTypes.h
@@ -0,0 +1,66 @@
+﻿// Copyright 1998-2014 Epic Games, Inc. All Rights Reserved.
+
+#pragma once
+
+#include "SkillSystemTypes.generated.h"
+
+class UAttributeComponent;
+
+UENUM(BlueprintType)
+namespace EGameplayAbilityActivationMode
+{
+	enum Type
+	{
+		Authority,			// We are the authority activating this ability
+		NonAuthority,		// We are not the authority but aren't predicting yet. This is a mostly invalid state to be in.
+
+		Predicting,			// We are predicting the activation of this ability
+		Confirmed,			// We are not the authority, but the authority has confirmed this activation
+	};
+}
+
+USTRUCT(BlueprintType)
+struct SKILLSYSTEM_API FGameplayAbilityActorInfo
+{
+	GENERATED_USTRUCT_BODY()
+
+	FGameplayAbilityActorInfo()
+		: ActivationMode(EGameplayAbilityActivationMode::Authority)
+		, PredictionKey(0)
+	{
+
+	}
+
+	FGameplayAbilityActorInfo(EGameplayAbilityActivationMode::Type InType)
+		: ActivationMode(InType)
+		, PredictionKey(0)
+	{
+	}
+
+
+	UPROPERTY(BlueprintReadOnly, Category = "ActorInfo")
+	TWeakObjectPtr<AActor>	Actor;
+
+	UPROPERTY(BlueprintReadOnly, Category = "ActorInfo")
+	TWeakObjectPtr<UAnimInstance>	AnimInstance;
+
+	UPROPERTY(BlueprintReadOnly, Category = "ActorInfo")
+	TWeakObjectPtr<UAttributeComponent>	AttributeComponent;
+
+	UPROPERTY(BlueprintReadOnly, Category = "ActorInfo")
+	TWeakObjectPtr<UMovementComponent>	MovementComponent;
+
+	UPROPERTY(BlueprintReadOnly, Category = "ActorInfo")
+	mutable TEnumAsByte<EGameplayAbilityActivationMode::Type>	ActivationMode;
+
+	UPROPERTY()
+	mutable int32 PredictionKey;
+
+	void InitFromActor(AActor *Actor);
+
+	void GeneratePredictionKey() const;
+
+	void SetPredictionKey(int32 InPredictionKey);
+
+	void SetActivationConfirmed();
+};
diff --git a/Engine/Source/Runtime/SkillSystem/SkillSystem.Build.cs b/Engine/Source/Runtime/SkillSystem/SkillSystem.Build.cs
--- a/Engine/Source/Runtime/SkillSystem/SkillSystem.Build.cs
+++ b/Engine/Source/Runtime/SkillSystem/SkillSystem.Build.cs
@@ -6,6 +6,8 @@ namespace UnrealBuildTool.Rules
 	{
         public SkillSystem(TargetInfo Target)
 		{
+            PrivateIncludePaths.Add("../../Engine/Source/Runtime/SkillSystem/Private");
+            
 			PublicIncludePaths.AddRange(
 				new string[] {
 					// ... add public include paths required here ...
