cis fix

diff --git a/Engine/Source/Runtime/SkillSystem/Private/SkillSystemDebugHUD.cpp b/Engine/Source/Runtime/SkillSystem/Private/SkillSystemDebugHUD.cpp
--- a/Engine/Source/Runtime/SkillSystem/Private/SkillSystemDebugHUD.cpp
+++ b/Engine/Source/Runtime/SkillSystem/Private/SkillSystemDebugHUD.cpp
@@ -107,7 +107,7 @@ void ASkillSystemDebugHUD::DrawDebugAttributeComponent(UAttributeComponent *Comp
 			void *PropertyValue = Prop->ContainerPtrToValuePtr<void>(Set);
 			Prop->ExportTextItem(ValueString, PropertyValue, NULL, NULL, 0);
 
-			String = FString::Printf(TEXT("%s: %s"), *Prop->GetDisplayNameText().ToString(), *ValueString);
+			String = FString::Printf(TEXT("%s: %s"), *Prop->GetName().ToString(), *ValueString);
 			DrawWithBackground(Font, String, Color, EAlignHorizontal::Left, X, EAlignVertical::Top, Y);
 		}
 
