Fixing linker error in monolithic builds

diff --git a/Engine/Source/Runtime/SkillSystem/Private/SkillSystemDebugHUD.cpp b/Engine/Source/Runtime/SkillSystem/Private/SkillSystemDebugHUD.cpp
--- a/Engine/Source/Runtime/SkillSystem/Private/SkillSystemDebugHUD.cpp
+++ b/Engine/Source/Runtime/SkillSystem/Private/SkillSystemDebugHUD.cpp
@@ -153,7 +153,7 @@ static void	ToggleDebugHUD(const TArray<FString>& Args, UWorld* InWorld)
 	}
 }
 
-FAutoConsoleCommandWithWorldAndArgs TestObjectRefSerializeCommand(
+FAutoConsoleCommandWithWorldAndArgs SkillSystemToggleDebugHUDCommand(
 	TEXT("SkillSystem.ToggleDebugHUD"),
 	TEXT("ToggleDebugHUD Drawing"),
 	FConsoleCommandWithWorldAndArgsDelegate::CreateStatic(ToggleDebugHUD)
