Orgnaize ability tasks into subfolder of Abilities

diff --git a/Engine/Source/Runtime/SkillSystem/Classes/Abilities/Tasks/BlueprintPlayMontageAndWaitTaskProxy.h b/Engine/Source/Runtime/SkillSystem/Classes/Abilities/Tasks/BlueprintPlayMontageAndWaitTaskProxy.h
new file mode 100644
--- /dev/null
+++ b/Engine/Source/Runtime/SkillSystem/Classes/Abilities/Tasks/BlueprintPlayMontageAndWaitTaskProxy.h
@@ -0,0 +1,24 @@
+// Copyright 1998-2014 Epic Games, Inc. All Rights Reserved.
+#pragma once
+#include "AbilityTaskProxyBase.h"
+#include "BlueprintPlayMontageAndWaitTaskProxy.generated.h"
+
+DECLARE_DYNAMIC_MULTICAST_DELEGATE(FMontageWaitSimpleDelegate);
+
+UCLASS(MinimalAPI)
+class UBlueprintPlayMontageAndWaitTaskProxy : public UAbilityTaskProxyBase
+{
+	GENERATED_UCLASS_BODY()
+
+	UPROPERTY(BlueprintAssignable)
+	FMontageWaitSimpleDelegate	OnComplete;
+
+	UPROPERTY(BlueprintAssignable)
+	FMontageWaitSimpleDelegate	OnInterrupted;
+
+	UFUNCTION()
+	void OnMontageEnded(UAnimMontage* Montage, bool bInterrupted);
+
+public:
+	
+};
\ No newline at end of file
diff --git a/Engine/Source/Runtime/SkillSystem/Classes/Abilities/Tasks/BlueprintWaitMovementModeChangeTaskProxy.h b/Engine/Source/Runtime/SkillSystem/Classes/Abilities/Tasks/BlueprintWaitMovementModeChangeTaskProxy.h
new file mode 100644
--- /dev/null
+++ b/Engine/Source/Runtime/SkillSystem/Classes/Abilities/Tasks/BlueprintWaitMovementModeChangeTaskProxy.h
@@ -0,0 +1,23 @@
+// Copyright 1998-2014 Epic Games, Inc. All Rights Reserved.
+#pragma once
+#include "AbilityTaskProxyBase.h"
+#include "BlueprintWaitMovementModeChangeTaskProxy.generated.h"
+
+DECLARE_DYNAMIC_MULTICAST_DELEGATE_OneParam(FMovementModeChangedDelegate, EMovementMode, NewMovementMode);
+
+UCLASS(MinimalAPI)
+class UBlueprintWaitMovementModeChangeTaskProxy : public UAbilityTaskProxyBase
+{
+	GENERATED_UCLASS_BODY()
+
+	UPROPERTY(BlueprintAssignable)
+	FMovementModeChangedDelegate	OnChange;
+
+	UFUNCTION()
+	void OnMovementModeChange(ACharacter * Character, EMovementMode PrevMovementMode, uint8 PreviousCustomMode);
+
+	EMovementMode	RequiredMode;
+
+public:
+	
+};
\ No newline at end of file
diff --git a/Engine/Source/Runtime/SkillSystem/Classes/BlueprintPlayMontageAndWaitTaskProxy.h b/Engine/Source/Runtime/SkillSystem/Classes/BlueprintPlayMontageAndWaitTaskProxy.h
deleted file mode 100644
diff --git a/Engine/Source/Runtime/SkillSystem/Classes/BlueprintWaitMovementModeChangeTaskProxy.h b/Engine/Source/Runtime/SkillSystem/Classes/BlueprintWaitMovementModeChangeTaskProxy.h
deleted file mode 100644
diff --git a/Engine/Source/Runtime/SkillSystem/Private/Abilities/Tasks/BlueprintPlayMontageAndWaitTaskProxy.cpp b/Engine/Source/Runtime/SkillSystem/Private/Abilities/Tasks/BlueprintPlayMontageAndWaitTaskProxy.cpp
new file mode 100644
--- /dev/null
+++ b/Engine/Source/Runtime/SkillSystem/Private/Abilities/Tasks/BlueprintPlayMontageAndWaitTaskProxy.cpp
@@ -0,0 +1,20 @@
+
+#include "SkillSystemModulePrivatePCH.h"
+
+UBlueprintPlayMontageAndWaitTaskProxy::UBlueprintPlayMontageAndWaitTaskProxy(const class FPostConstructInitializeProperties& PCIP)
+	: Super(PCIP)
+{
+
+}
+
+void UBlueprintPlayMontageAndWaitTaskProxy::OnMontageEnded(UAnimMontage* Montage, bool bInterrupted)
+{
+	if (bInterrupted)
+	{
+		OnInterrupted.Broadcast();
+	}
+	else
+	{
+		OnComplete.Broadcast();
+	}
+}
\ No newline at end of file
diff --git a/Engine/Source/Runtime/SkillSystem/Private/Abilities/Tasks/BlueprintWaitMovementModeChangeTaskProxy.cpp b/Engine/Source/Runtime/SkillSystem/Private/Abilities/Tasks/BlueprintWaitMovementModeChangeTaskProxy.cpp
new file mode 100644
--- /dev/null
+++ b/Engine/Source/Runtime/SkillSystem/Private/Abilities/Tasks/BlueprintWaitMovementModeChangeTaskProxy.cpp
@@ -0,0 +1,24 @@
+
+#include "SkillSystemModulePrivatePCH.h"
+
+UBlueprintWaitMovementModeChangeTaskProxy::UBlueprintWaitMovementModeChangeTaskProxy(const class FPostConstructInitializeProperties& PCIP)
+	: Super(PCIP)
+{
+	RequiredMode = MOVE_None;
+}
+
+void UBlueprintWaitMovementModeChangeTaskProxy::OnMovementModeChange(ACharacter * Character, EMovementMode PrevMovementMode, uint8 PreviousCustomMode)
+{
+	if (Character)
+	{
+		if (UCharacterMovementComponent *MoveComp = Cast<UCharacterMovementComponent>(Character->GetMovementComponent()))
+		{
+			if (RequiredMode == MOVE_None || MoveComp->MovementMode == RequiredMode)
+			{
+				Character->MovementModeChangedDelegate.RemoveDynamic(this, &UBlueprintWaitMovementModeChangeTaskProxy::OnMovementModeChange);
+				OnChange.Broadcast(MoveComp->MovementMode);
+				return;
+			}
+		}
+	}
+}
\ No newline at end of file
diff --git a/Engine/Source/Runtime/SkillSystem/Private/BlueprintPlayMontageAndWaitTaskProxy.cpp b/Engine/Source/Runtime/SkillSystem/Private/BlueprintPlayMontageAndWaitTaskProxy.cpp
deleted file mode 100644
diff --git a/Engine/Source/Runtime/SkillSystem/Private/BlueprintWaitMovementModeChangeTaskProxy.cpp b/Engine/Source/Runtime/SkillSystem/Private/BlueprintWaitMovementModeChangeTaskProxy.cpp
deleted file mode 100644
