SkillSystem - remove SkillSystemClasses.h from PCH, manually include classes where needed.

diff --git a/Engine/Source/Editor/SkillSystemEditor/Private/K2Node_PlayMontageAndWait.cpp b/Engine/Source/Editor/SkillSystemEditor/Private/K2Node_PlayMontageAndWait.cpp
--- a/Engine/Source/Editor/SkillSystemEditor/Private/K2Node_PlayMontageAndWait.cpp
+++ b/Engine/Source/Editor/SkillSystemEditor/Private/K2Node_PlayMontageAndWait.cpp
@@ -3,7 +3,7 @@
 
 #include "SkillSystemEditorModulePrivatePCH.h"
 #include "SkillSystemBlueprintLibrary.h"
-#include "BlueprintPlayMontageAndWaitTaskProxy.h"
+#include "Abilities/Tasks/BlueprintPlayMontageAndWaitTaskProxy.h"
 #include "CompilerResultsLog.h"
 
 #define LOCTEXT_NAMESPACE "K2Node_PlayMontageAndWait"
diff --git a/Engine/Source/Editor/SkillSystemEditor/Private/K2Node_WaitMovementModeChange.cpp b/Engine/Source/Editor/SkillSystemEditor/Private/K2Node_WaitMovementModeChange.cpp
--- a/Engine/Source/Editor/SkillSystemEditor/Private/K2Node_WaitMovementModeChange.cpp
+++ b/Engine/Source/Editor/SkillSystemEditor/Private/K2Node_WaitMovementModeChange.cpp
@@ -3,7 +3,7 @@
 
 #include "SkillSystemEditorModulePrivatePCH.h"
 #include "SkillSystemBlueprintLibrary.h"
-#include "BlueprintWaitMovementModeChangeTaskProxy.h"
+#include "Abilities/Tasks/BlueprintWaitMovementModeChangeTaskProxy.h"
 #include "CompilerResultsLog.h"
 
 #define LOCTEXT_NAMESPACE "K2Node_WaitMovementModeChange"
diff --git a/Engine/Source/Runtime/SkillSystem/Classes/AttributeSet.h b/Engine/Source/Runtime/SkillSystem/Classes/AttributeSet.h
--- a/Engine/Source/Runtime/SkillSystem/Classes/AttributeSet.h
+++ b/Engine/Source/Runtime/SkillSystem/Classes/AttributeSet.h
@@ -1,8 +1,7 @@
-// Copyright 1998-2013 Epic Games, Inc. All Rights Reserved.
+// Copyright 1998-2014 Epic Games, Inc. All Rights Reserved.
 
 #pragma once
 
-//#include "../Engine/DataTable.h"
 #include "AttributeSet.generated.h"
 
 USTRUCT(BlueprintType)
diff --git a/Engine/Source/Runtime/SkillSystem/Private/Abilities/GameplayAbility_CharacterJump.cpp b/Engine/Source/Runtime/SkillSystem/Private/Abilities/GameplayAbility_CharacterJump.cpp
--- a/Engine/Source/Runtime/SkillSystem/Private/Abilities/GameplayAbility_CharacterJump.cpp
+++ b/Engine/Source/Runtime/SkillSystem/Private/Abilities/GameplayAbility_CharacterJump.cpp
@@ -1,6 +1,7 @@
 // Copyright 1998-2014 Epic Games, Inc. All Rights Reserved.
 
 #include "SkillSystemModulePrivatePCH.h"
+#include "Abilities/GameplayAbility_CharacterJump.h"
 
 // --------------------------------------------------------------------------------------------------------------------------------------------------------
 //
diff --git a/Engine/Source/Runtime/SkillSystem/Private/Abilities/GameplayAbility_Montage.cpp b/Engine/Source/Runtime/SkillSystem/Private/Abilities/GameplayAbility_Montage.cpp
--- a/Engine/Source/Runtime/SkillSystem/Private/Abilities/GameplayAbility_Montage.cpp
+++ b/Engine/Source/Runtime/SkillSystem/Private/Abilities/GameplayAbility_Montage.cpp
@@ -1,6 +1,7 @@
 // Copyright 1998-2013 Epic Games, Inc. All Rights Reserved.
 
 #include "SkillSystemModulePrivatePCH.h"
+#include "Abilities/GameplayAbility_Montage.h"
 
 // --------------------------------------------------------------------------------------------------------------------------------------------------------
 //
diff --git a/Engine/Source/Runtime/SkillSystem/Private/AbilityTaskProxyBase.cpp b/Engine/Source/Runtime/SkillSystem/Private/AbilityTaskProxyBase.cpp
--- a/Engine/Source/Runtime/SkillSystem/Private/AbilityTaskProxyBase.cpp
+++ b/Engine/Source/Runtime/SkillSystem/Private/AbilityTaskProxyBase.cpp
@@ -1,5 +1,6 @@
 
 #include "SkillSystemModulePrivatePCH.h"
+#include "AbilityTaskProxyBase.h"
 
 UAbilityTaskProxyBase::UAbilityTaskProxyBase(const class FPostConstructInitializeProperties& PCIP)
 	: Super(PCIP)
diff --git a/Engine/Source/Runtime/SkillSystem/Private/AttributeComponent.cpp b/Engine/Source/Runtime/SkillSystem/Private/AttributeComponent.cpp
--- a/Engine/Source/Runtime/SkillSystem/Private/AttributeComponent.cpp
+++ b/Engine/Source/Runtime/SkillSystem/Private/AttributeComponent.cpp
@@ -2,6 +2,10 @@
 // ActorComponent.cpp: Actor component implementation.
 
 #include "SkillSystemModulePrivatePCH.h"
+#include "AttributeComponent.h"
+#include "GameplayCueInterface.h"
+#include "Abilities/GameplayAbility.h"
+
 
 #include "Net/UnrealNetwork.h"
 #include "Engine/ActorChannel.h"
@@ -235,7 +239,7 @@ FActiveGameplayEffectHandle UAttributeComponent::ApplyGameplayEffectSpecToTarget
 FActiveGameplayEffectHandle UAttributeComponent::ApplyGameplayEffectSpecToSelf(const FGameplayEffectSpec &Spec, FModifierQualifier BaseQualifier)
 {
 	// Temp, only non instance, non periodic GEs can be predictive
-	check((BaseQualifier.PredictionKey() >= 0) == (Spec.GetDuration() != UGameplayEffect::INSTANT_APPLICATION && Spec.GetPeriod() == UGameplayEffect::NO_PERIOD));
+	check((BaseQualifier.PredictionKey() > 0) == (Spec.GetDuration() != UGameplayEffect::INSTANT_APPLICATION && Spec.GetPeriod() == UGameplayEffect::NO_PERIOD));
 
 	// check if the effect being applied actually succeeds
 	float ChanceToApply = Spec.GetChanceToApplyToTarget();
diff --git a/Engine/Source/Runtime/SkillSystem/Private/AttributeComponent_Abilities.cpp b/Engine/Source/Runtime/SkillSystem/Private/AttributeComponent_Abilities.cpp
--- a/Engine/Source/Runtime/SkillSystem/Private/AttributeComponent_Abilities.cpp
+++ b/Engine/Source/Runtime/SkillSystem/Private/AttributeComponent_Abilities.cpp
@@ -2,6 +2,7 @@
 // ActorComponent.cpp: Actor component implementation.
 
 #include "SkillSystemModulePrivatePCH.h"
+#include "Abilities/GameplayAbility.h"
 
 #include "Net/UnrealNetwork.h"
 #include "MessageLog.h"
diff --git a/Engine/Source/Runtime/SkillSystem/Private/GameplayAbilitySet.cpp b/Engine/Source/Runtime/SkillSystem/Private/GameplayAbilitySet.cpp
--- a/Engine/Source/Runtime/SkillSystem/Private/GameplayAbilitySet.cpp
+++ b/Engine/Source/Runtime/SkillSystem/Private/GameplayAbilitySet.cpp
@@ -1,6 +1,7 @@
 // Copyright 1998-2014 Epic Games, Inc. All Rights Reserved.
 
 #include "SkillSystemModulePrivatePCH.h"
+#include "GameplayAbilitySet.h"
 
 // --------------------------------------------------------------------------------------------------------------------------------------------------------
 //
diff --git a/Engine/Source/Runtime/SkillSystem/Private/GameplayCueView.cpp b/Engine/Source/Runtime/SkillSystem/Private/GameplayCueView.cpp
--- a/Engine/Source/Runtime/SkillSystem/Private/GameplayCueView.cpp
+++ b/Engine/Source/Runtime/SkillSystem/Private/GameplayCueView.cpp
@@ -3,6 +3,7 @@
 #pragma once
 
 #include "SkillSystemModulePrivatePCH.h"
+#include "GameplayCueView.h"
 #include "ParticleDefinitions.h"
 #include "SoundDefinitions.h"
 #include "Particles/ParticleLODLevel.h"
diff --git a/Engine/Source/Runtime/SkillSystem/Private/GameplayEffect.cpp b/Engine/Source/Runtime/SkillSystem/Private/GameplayEffect.cpp
--- a/Engine/Source/Runtime/SkillSystem/Private/GameplayEffect.cpp
+++ b/Engine/Source/Runtime/SkillSystem/Private/GameplayEffect.cpp
@@ -1,6 +1,9 @@
 // Copyright 1998-2013 Epic Games, Inc. All Rights Reserved.
 
 #include "SkillSystemModulePrivatePCH.h"
+#include "GameplayEffect.h"
+#include "GameplayEffectExtension.h"
+#include "GameplayEffectStackingExtension.h"
 
 const float UGameplayEffect::INFINITE_DURATION = -1.f;
 const float UGameplayEffect::INSTANT_APPLICATION = 0.f;
diff --git a/Engine/Source/Runtime/SkillSystem/Private/GameplayEffectExtension_LifestealTest.cpp b/Engine/Source/Runtime/SkillSystem/Private/GameplayEffectExtension_LifestealTest.cpp
--- a/Engine/Source/Runtime/SkillSystem/Private/GameplayEffectExtension_LifestealTest.cpp
+++ b/Engine/Source/Runtime/SkillSystem/Private/GameplayEffectExtension_LifestealTest.cpp
@@ -1,6 +1,8 @@
 // Copyright 1998-2013 Epic Games, Inc. All Rights Reserved.
 
 #include "SkillSystemModulePrivatePCH.h"
+#include "GameplayEffectExtension_LifestealTest.h"
+#include "SkillSystemTestAttributeSet.h"
 #include "GameplayTagsModule.h"
 
 UGameplayEffectExtension_LifestealTest::UGameplayEffectExtension_LifestealTest(const class FPostConstructInitializeProperties& PCIP)
diff --git a/Engine/Source/Runtime/SkillSystem/Private/GameplayEffectExtension_ShieldTest.cpp b/Engine/Source/Runtime/SkillSystem/Private/GameplayEffectExtension_ShieldTest.cpp
--- a/Engine/Source/Runtime/SkillSystem/Private/GameplayEffectExtension_ShieldTest.cpp
+++ b/Engine/Source/Runtime/SkillSystem/Private/GameplayEffectExtension_ShieldTest.cpp
@@ -1,6 +1,7 @@
 // Copyright 1998-2013 Epic Games, Inc. All Rights Reserved.
 
 #include "SkillSystemModulePrivatePCH.h"
+#include "GameplayEffectExtension_ShieldTest.h"
 #include "GameplayTagsModule.h"
 
 UGameplayEffectExtension_ShieldTest::UGameplayEffectExtension_ShieldTest(const class FPostConstructInitializeProperties& PCIP)
diff --git a/Engine/Source/Runtime/SkillSystem/Private/GameplayEffectStackingExtension_CappedNumberTest.cpp b/Engine/Source/Runtime/SkillSystem/Private/GameplayEffectStackingExtension_CappedNumberTest.cpp
--- a/Engine/Source/Runtime/SkillSystem/Private/GameplayEffectStackingExtension_CappedNumberTest.cpp
+++ b/Engine/Source/Runtime/SkillSystem/Private/GameplayEffectStackingExtension_CappedNumberTest.cpp
@@ -1,6 +1,7 @@
 // Copyright 1998-2014 Epic Games, Inc. All Rights Reserved.
 
 #include "SkillSystemModulePrivatePCH.h"
+#include "GameplayEffectStackingExtension_CappedNumberTest.h"
 #include "GameplayTagsModule.h"
 #include "Crc.h"
 
diff --git a/Engine/Source/Runtime/SkillSystem/Private/GameplayEffectStackingExtension_DiminishingReturnsTest.cpp b/Engine/Source/Runtime/SkillSystem/Private/GameplayEffectStackingExtension_DiminishingReturnsTest.cpp
--- a/Engine/Source/Runtime/SkillSystem/Private/GameplayEffectStackingExtension_DiminishingReturnsTest.cpp
+++ b/Engine/Source/Runtime/SkillSystem/Private/GameplayEffectStackingExtension_DiminishingReturnsTest.cpp
@@ -1,6 +1,7 @@
 // Copyright 1998-2014 Epic Games, Inc. All Rights Reserved.
 
 #include "SkillSystemModulePrivatePCH.h"
+#include "GameplayEffectStackingExtension_DiminishingReturnsTest.h"
 #include "GameplayTagsModule.h"
 
 UCurveTable* GetCurveTable()
diff --git a/Engine/Source/Runtime/SkillSystem/Private/GameplayEffectTests.cpp b/Engine/Source/Runtime/SkillSystem/Private/GameplayEffectTests.cpp
--- a/Engine/Source/Runtime/SkillSystem/Private/GameplayEffectTests.cpp
+++ b/Engine/Source/Runtime/SkillSystem/Private/GameplayEffectTests.cpp
@@ -1,6 +1,9 @@
 // Copyright 1998-2013 Epic Games, Inc. All Rights Reserved.
 
 #include "SkillSystemModulePrivatePCH.h"
+#include "SkillSystemTestPawn.h"
+#include "GameplayEffect.h"
+#include "AttributeSet.h"
 #include "GameplayTagsModule.h"
 //#include "SkillSystemTestPawn.h"
 
diff --git a/Engine/Source/Runtime/SkillSystem/Private/SkillSystemBlueprintLibrary.cpp b/Engine/Source/Runtime/SkillSystem/Private/SkillSystemBlueprintLibrary.cpp
--- a/Engine/Source/Runtime/SkillSystem/Private/SkillSystemBlueprintLibrary.cpp
+++ b/Engine/Source/Runtime/SkillSystem/Private/SkillSystemBlueprintLibrary.cpp
@@ -1,6 +1,11 @@
 // Copyright 1998-2014 Epic Games, Inc. All Rights Reserved.
 
 #include "SkillSystemModulePrivatePCH.h"
+#include "SkillSystemBlueprintLibrary.h"
+#include "Abilities/GameplayAbility.h"
+#include "Abilities/GameplayAbility_Instanced.h"
+#include "Abilities/Tasks/BlueprintPlayMontageAndWaitTaskProxy.h"
+#include "Abilities/Tasks/BlueprintWaitMovementModeChangeTaskProxy.h"
 #include "LatentActions.h"
 
 USkillSystemBlueprintLibrary::USkillSystemBlueprintLibrary(const class FPostConstructInitializeProperties& PCIP)
diff --git a/Engine/Source/Runtime/SkillSystem/Private/SkillSystemModulePrivatePCH.h b/Engine/Source/Runtime/SkillSystem/Private/SkillSystemModulePrivatePCH.h
--- a/Engine/Source/Runtime/SkillSystem/Private/SkillSystemModulePrivatePCH.h
+++ b/Engine/Source/Runtime/SkillSystem/Private/SkillSystemModulePrivatePCH.h
@@ -16,11 +16,10 @@
 
 #include "SkillSystemTypes.h"
 
-#include "SkillSystemClasses.h"
+//#include "SkillSystemClasses.h"
 
 #include "SkillSystem.h"
 #include "SkillSystemModule.h"
 #include "SkillSystemLog.h"
 
 #include "PlayMontageAndWaitAction.h"
-#include "BlueprintPlayMontageAndWaitTaskProxy.h"
\ No newline at end of file
