Rename Ability TaskProxy classes

diff --git a/Engine/Source/Editor/SkillSystemEditor/Private/K2Node_PlayMontageAndWait.cpp b/Engine/Source/Editor/SkillSystemEditor/Private/K2Node_PlayMontageAndWait.cpp
--- a/Engine/Source/Editor/SkillSystemEditor/Private/K2Node_PlayMontageAndWait.cpp
+++ b/Engine/Source/Editor/SkillSystemEditor/Private/K2Node_PlayMontageAndWait.cpp
@@ -3,7 +3,7 @@
 
 #include "SkillSystemEditorModulePrivatePCH.h"
 #include "SkillSystemBlueprintLibrary.h"
-#include "Abilities/Tasks/BlueprintPlayMontageAndWaitTaskProxy.h"
+#include "Abilities/Tasks/AbilityTask_PlayMontageAndWait.h"
 #include "CompilerResultsLog.h"
 
 #define LOCTEXT_NAMESPACE "K2Node_PlayMontageAndWait"
@@ -14,7 +14,7 @@ UK2Node_PlayMontageAndWait::UK2Node_PlayMontageAndWait(const class FPostConstruc
 	
 	ProxyFactoryFunctionName = GET_FUNCTION_NAME_CHECKED(USkillSystemBlueprintLibrary, CreatePlayMontageAndWaitProxy);
 	ProxyFactoryClass = USkillSystemBlueprintLibrary::StaticClass();
-	ProxyClass = UBlueprintPlayMontageAndWaitTaskProxy::StaticClass();
+	ProxyClass = UAbilityTask_PlayMontageAndWait::StaticClass();
 }
 
 FString UK2Node_PlayMontageAndWait::GetCategoryName()
diff --git a/Engine/Source/Editor/SkillSystemEditor/Private/K2Node_WaitMovementModeChange.cpp b/Engine/Source/Editor/SkillSystemEditor/Private/K2Node_WaitMovementModeChange.cpp
--- a/Engine/Source/Editor/SkillSystemEditor/Private/K2Node_WaitMovementModeChange.cpp
+++ b/Engine/Source/Editor/SkillSystemEditor/Private/K2Node_WaitMovementModeChange.cpp
@@ -3,7 +3,7 @@
 
 #include "SkillSystemEditorModulePrivatePCH.h"
 #include "SkillSystemBlueprintLibrary.h"
-#include "Abilities/Tasks/BlueprintWaitMovementModeChangeTaskProxy.h"
+#include "Abilities/Tasks/AbilityTask_WaitMovementModeChange.h"
 #include "CompilerResultsLog.h"
 
 #define LOCTEXT_NAMESPACE "K2Node_WaitMovementModeChange"
@@ -14,7 +14,7 @@ UK2Node_WaitMovementModeChange::UK2Node_WaitMovementModeChange(const class FPost
 	
 	ProxyFactoryFunctionName = GET_FUNCTION_NAME_CHECKED(USkillSystemBlueprintLibrary, CreateWaitMovementModeChange);
 	ProxyFactoryClass = USkillSystemBlueprintLibrary::StaticClass();
-	ProxyClass = UBlueprintWaitMovementModeChangeTaskProxy::StaticClass();
+	ProxyClass = UAbilityTask_WaitMovementModeChange::StaticClass();
 }
 
 FString UK2Node_WaitMovementModeChange::GetCategoryName()
diff --git a/Engine/Source/Runtime/SkillSystem/Classes/Abilities/Tasks/AbilityTask.h b/Engine/Source/Runtime/SkillSystem/Classes/Abilities/Tasks/AbilityTask.h
new file mode 100644
--- /dev/null
+++ b/Engine/Source/Runtime/SkillSystem/Classes/Abilities/Tasks/AbilityTask.h
@@ -0,0 +1,12 @@
+// Copyright 1998-2014 Epic Games, Inc. All Rights Reserved.
+#pragma once
+#include "AbilityTask.Generated.h"
+
+UCLASS(MinimalAPI)
+class UAbilityTask : public UObject
+{
+	GENERATED_UCLASS_BODY()	
+
+public:
+	
+};
\ No newline at end of file
diff --git a/Engine/Source/Runtime/SkillSystem/Classes/Abilities/Tasks/AbilityTask_PlayMontageAndWait.h b/Engine/Source/Runtime/SkillSystem/Classes/Abilities/Tasks/AbilityTask_PlayMontageAndWait.h
new file mode 100644
--- /dev/null
+++ b/Engine/Source/Runtime/SkillSystem/Classes/Abilities/Tasks/AbilityTask_PlayMontageAndWait.h
@@ -0,0 +1,24 @@
+// Copyright 1998-2014 Epic Games, Inc. All Rights Reserved.
+#pragma once
+#include "AbilityTask.h"
+#include "AbilityTask_PlayMontageAndWait.generated.h"
+
+DECLARE_DYNAMIC_MULTICAST_DELEGATE(FMontageWaitSimpleDelegate);
+
+UCLASS(MinimalAPI)
+class UAbilityTask_PlayMontageAndWait : public UAbilityTask
+{
+	GENERATED_UCLASS_BODY()
+
+	UPROPERTY(BlueprintAssignable)
+	FMontageWaitSimpleDelegate	OnComplete;
+
+	UPROPERTY(BlueprintAssignable)
+	FMontageWaitSimpleDelegate	OnInterrupted;
+
+	UFUNCTION()
+	void OnMontageEnded(UAnimMontage* Montage, bool bInterrupted);
+
+public:
+	
+};
\ No newline at end of file
diff --git a/Engine/Source/Runtime/SkillSystem/Classes/Abilities/Tasks/AbilityTask_WaitMovementModeChange.h b/Engine/Source/Runtime/SkillSystem/Classes/Abilities/Tasks/AbilityTask_WaitMovementModeChange.h
new file mode 100644
--- /dev/null
+++ b/Engine/Source/Runtime/SkillSystem/Classes/Abilities/Tasks/AbilityTask_WaitMovementModeChange.h
@@ -0,0 +1,23 @@
+// Copyright 1998-2014 Epic Games, Inc. All Rights Reserved.
+#pragma once
+#include "AbilityTask.h"
+#include "AbilityTask_WaitMovementModeChange.generated.h"
+
+DECLARE_DYNAMIC_MULTICAST_DELEGATE_OneParam(FMovementModeChangedDelegate, EMovementMode, NewMovementMode);
+
+UCLASS(MinimalAPI)
+class UAbilityTask_WaitMovementModeChange : public UAbilityTask
+{
+	GENERATED_UCLASS_BODY()
+
+	UPROPERTY(BlueprintAssignable)
+	FMovementModeChangedDelegate	OnChange;
+
+	UFUNCTION()
+	void OnMovementModeChange(ACharacter * Character, EMovementMode PrevMovementMode, uint8 PreviousCustomMode);
+
+	EMovementMode	RequiredMode;
+
+public:
+	
+};
\ No newline at end of file
diff --git a/Engine/Source/Runtime/SkillSystem/Classes/Abilities/Tasks/BlueprintPlayMontageAndWaitTaskProxy.h b/Engine/Source/Runtime/SkillSystem/Classes/Abilities/Tasks/BlueprintPlayMontageAndWaitTaskProxy.h
deleted file mode 100644
diff --git a/Engine/Source/Runtime/SkillSystem/Classes/Abilities/Tasks/BlueprintWaitMovementModeChangeTaskProxy.h b/Engine/Source/Runtime/SkillSystem/Classes/Abilities/Tasks/BlueprintWaitMovementModeChangeTaskProxy.h
deleted file mode 100644
diff --git a/Engine/Source/Runtime/SkillSystem/Classes/AbilityTaskProxyBase.h b/Engine/Source/Runtime/SkillSystem/Classes/AbilityTaskProxyBase.h
deleted file mode 100644
diff --git a/Engine/Source/Runtime/SkillSystem/Classes/SkillSystemBlueprintLibrary.h b/Engine/Source/Runtime/SkillSystem/Classes/SkillSystemBlueprintLibrary.h
--- a/Engine/Source/Runtime/SkillSystem/Classes/SkillSystemBlueprintLibrary.h
+++ b/Engine/Source/Runtime/SkillSystem/Classes/SkillSystemBlueprintLibrary.h
@@ -17,8 +17,8 @@ class USkillSystemBlueprintLibrary : public UBlueprintFunctionLibrary
 
 
 	UFUNCTION(BlueprintCallable, meta=(HidePin="WorldContextObject", DefaultToSelf="WorldContextObject", BlueprintInternalUseOnly = "TRUE"))
-	static UBlueprintPlayMontageAndWaitTaskProxy* CreatePlayMontageAndWaitProxy(class UObject* WorldContextObject, class UAnimMontage *MontageToPlay);
+	static class UAbilityTask_PlayMontageAndWait* CreatePlayMontageAndWaitProxy(class UObject* WorldContextObject, class UAnimMontage *MontageToPlay);
 
 	UFUNCTION(BlueprintCallable, meta = (HidePin = "WorldContextObject", DefaultToSelf = "WorldContextObject", BlueprintInternalUseOnly = "TRUE"))
-	static class UBlueprintWaitMovementModeChangeTaskProxy* CreateWaitMovementModeChange(class UObject* WorldContextObject, EMovementMode NewMode);
+	static class UAbilityTask_WaitMovementModeChange* CreateWaitMovementModeChange(class UObject* WorldContextObject, EMovementMode NewMode);
 };
\ No newline at end of file
diff --git a/Engine/Source/Runtime/SkillSystem/Private/Abilities/Tasks/AbilityTask.cpp b/Engine/Source/Runtime/SkillSystem/Private/Abilities/Tasks/AbilityTask.cpp
new file mode 100644
--- /dev/null
+++ b/Engine/Source/Runtime/SkillSystem/Private/Abilities/Tasks/AbilityTask.cpp
@@ -0,0 +1,9 @@
+
+#include "SkillSystemModulePrivatePCH.h"
+#include "Abilities/Tasks/AbilityTask.h"
+
+UAbilityTask::UAbilityTask(const class FPostConstructInitializeProperties& PCIP)
+	: Super(PCIP)
+{
+
+}
\ No newline at end of file
diff --git a/Engine/Source/Runtime/SkillSystem/Private/Abilities/Tasks/AbilityTask_PlayMontageAndWait.cpp b/Engine/Source/Runtime/SkillSystem/Private/Abilities/Tasks/AbilityTask_PlayMontageAndWait.cpp
new file mode 100644
--- /dev/null
+++ b/Engine/Source/Runtime/SkillSystem/Private/Abilities/Tasks/AbilityTask_PlayMontageAndWait.cpp
@@ -0,0 +1,20 @@
+
+#include "SkillSystemModulePrivatePCH.h"
+
+UAbilityTask_PlayMontageAndWait::UAbilityTask_PlayMontageAndWait(const class FPostConstructInitializeProperties& PCIP)
+	: Super(PCIP)
+{
+
+}
+
+void UAbilityTask_PlayMontageAndWait::OnMontageEnded(UAnimMontage* Montage, bool bInterrupted)
+{
+	if (bInterrupted)
+	{
+		OnInterrupted.Broadcast();
+	}
+	else
+	{
+		OnComplete.Broadcast();
+	}
+}
\ No newline at end of file
diff --git a/Engine/Source/Runtime/SkillSystem/Private/Abilities/Tasks/AbilityTask_WaitMovementModeChange.cpp b/Engine/Source/Runtime/SkillSystem/Private/Abilities/Tasks/AbilityTask_WaitMovementModeChange.cpp
new file mode 100644
--- /dev/null
+++ b/Engine/Source/Runtime/SkillSystem/Private/Abilities/Tasks/AbilityTask_WaitMovementModeChange.cpp
@@ -0,0 +1,24 @@
+
+#include "SkillSystemModulePrivatePCH.h"
+
+UAbilityTask_WaitMovementModeChange::UAbilityTask_WaitMovementModeChange(const class FPostConstructInitializeProperties& PCIP)
+	: Super(PCIP)
+{
+	RequiredMode = MOVE_None;
+}
+
+void UAbilityTask_WaitMovementModeChange::OnMovementModeChange(ACharacter * Character, EMovementMode PrevMovementMode, uint8 PreviousCustomMode)
+{
+	if (Character)
+	{
+		if (UCharacterMovementComponent *MoveComp = Cast<UCharacterMovementComponent>(Character->GetMovementComponent()))
+		{
+			if (RequiredMode == MOVE_None || MoveComp->MovementMode == RequiredMode)
+			{
+				Character->MovementModeChangedDelegate.RemoveDynamic(this, &UAbilityTask_WaitMovementModeChange::OnMovementModeChange);
+				OnChange.Broadcast(MoveComp->MovementMode);
+				return;
+			}
+		}
+	}
+}
\ No newline at end of file
diff --git a/Engine/Source/Runtime/SkillSystem/Private/Abilities/Tasks/BlueprintPlayMontageAndWaitTaskProxy.cpp b/Engine/Source/Runtime/SkillSystem/Private/Abilities/Tasks/BlueprintPlayMontageAndWaitTaskProxy.cpp
deleted file mode 100644
diff --git a/Engine/Source/Runtime/SkillSystem/Private/Abilities/Tasks/BlueprintWaitMovementModeChangeTaskProxy.cpp b/Engine/Source/Runtime/SkillSystem/Private/Abilities/Tasks/BlueprintWaitMovementModeChangeTaskProxy.cpp
deleted file mode 100644
diff --git a/Engine/Source/Runtime/SkillSystem/Private/AbilityTaskProxyBase.cpp b/Engine/Source/Runtime/SkillSystem/Private/AbilityTaskProxyBase.cpp
deleted file mode 100644
diff --git a/Engine/Source/Runtime/SkillSystem/Private/SkillSystemBlueprintLibrary.cpp b/Engine/Source/Runtime/SkillSystem/Private/SkillSystemBlueprintLibrary.cpp
--- a/Engine/Source/Runtime/SkillSystem/Private/SkillSystemBlueprintLibrary.cpp
+++ b/Engine/Source/Runtime/SkillSystem/Private/SkillSystemBlueprintLibrary.cpp
@@ -4,8 +4,8 @@
 #include "SkillSystemBlueprintLibrary.h"
 #include "Abilities/GameplayAbility.h"
 #include "Abilities/GameplayAbility_Instanced.h"
-#include "Abilities/Tasks/BlueprintPlayMontageAndWaitTaskProxy.h"
-#include "Abilities/Tasks/BlueprintWaitMovementModeChangeTaskProxy.h"
+#include "Abilities/Tasks/AbilityTask_PlayMontageAndWait.h"
+#include "Abilities/Tasks/AbilityTask_WaitMovementModeChange.h"
 #include "LatentActions.h"
 
 USkillSystemBlueprintLibrary::USkillSystemBlueprintLibrary(const class FPostConstructInitializeProperties& PCIP)
@@ -42,7 +42,7 @@ void USkillSystemBlueprintLibrary::StaticFuncPlayMontageAndWait(AActor* WorldCon
 	}
 }
 
-UBlueprintPlayMontageAndWaitTaskProxy* USkillSystemBlueprintLibrary::CreatePlayMontageAndWaitProxy(class UObject* WorldContextObject, UAnimMontage *MontageToPlay)
+UAbilityTask_PlayMontageAndWait* USkillSystemBlueprintLibrary::CreatePlayMontageAndWaitProxy(class UObject* WorldContextObject, UAnimMontage *MontageToPlay)
 {
 	check(WorldContextObject);
 
@@ -51,11 +51,11 @@ UBlueprintPlayMontageAndWaitTaskProxy* USkillSystemBlueprintLibrary::CreatePlayM
 	{
 		AActor * ActorOwner = Cast<AActor>(Ability->GetOuter());
 
-		UBlueprintPlayMontageAndWaitTaskProxy * MyObj = NULL;
-		MyObj = NewObject<UBlueprintPlayMontageAndWaitTaskProxy>();
+		UAbilityTask_PlayMontageAndWait * MyObj = NULL;
+		MyObj = NewObject<UAbilityTask_PlayMontageAndWait>();
 
 		FOnMontageEnded EndDelegate;
-		EndDelegate.BindUObject(MyObj, &UBlueprintPlayMontageAndWaitTaskProxy::OnMontageEnded);
+		EndDelegate.BindUObject(MyObj, &UAbilityTask_PlayMontageAndWait::OnMontageEnded);
 
 		FGameplayAbilityActorInfo	ActorInfo;
 		ActorInfo.InitFromActor(ActorOwner);
@@ -70,7 +70,7 @@ UBlueprintPlayMontageAndWaitTaskProxy* USkillSystemBlueprintLibrary::CreatePlayM
 }
 
 
-UBlueprintWaitMovementModeChangeTaskProxy* USkillSystemBlueprintLibrary::CreateWaitMovementModeChange(class UObject* WorldContextObject, EMovementMode NewMode)
+UAbilityTask_WaitMovementModeChange* USkillSystemBlueprintLibrary::CreateWaitMovementModeChange(class UObject* WorldContextObject, EMovementMode NewMode)
 {
 	check(WorldContextObject);
 
@@ -79,8 +79,8 @@ UBlueprintWaitMovementModeChangeTaskProxy* USkillSystemBlueprintLibrary::CreateW
 	{
 		AActor * ActorOwner = Cast<AActor>(Ability->GetOuter());
 
-		UBlueprintWaitMovementModeChangeTaskProxy * MyObj = NULL;
-		MyObj = NewObject<UBlueprintWaitMovementModeChangeTaskProxy>();
+		UAbilityTask_WaitMovementModeChange * MyObj = NULL;
+		MyObj = NewObject<UAbilityTask_WaitMovementModeChange>();
 		MyObj->RequiredMode = NewMode;
 
 		FGameplayAbilityActorInfo	ActorInfo;
@@ -89,21 +89,8 @@ UBlueprintWaitMovementModeChangeTaskProxy* USkillSystemBlueprintLibrary::CreateW
 		ACharacter * Character = Cast<ACharacter>(ActorInfo.Actor.Get());
 		if (Character)
 		{
-			/*
-				FScriptDelegate Delegate;
-				Delegate.BindUFunction(this, "OnActorBump");
-				MyPawn->OnActorHit.AddUnique(Delegate);
-			*/
-
-			//FMovementModeChangedSignature::FDelegate::CreateUObject(MyObj, &UBlueprintWaitMovementModeChangeTaskProxy::OnMovementModeChange);
-			//MyDel.BindUObject( 
-			//Character->MovementModeChangedDelegate.Add(MyObj, &UBlueprintWaitMovementModeChangeTaskProxy::OnMovementModeChange); // //FMovementModeChangedSignature::FDelegate::CreateUObject
-
-			Character->MovementModeChangedDelegate.AddDynamic(MyObj, &UBlueprintWaitMovementModeChangeTaskProxy::OnMovementModeChange);
-
-			//Character->MovementModeChangedDelegate.AddUObject(MyObj, &UBlueprintWaitMovementModeChangeTaskProxy::OnMovementModeChange);
+			Character->MovementModeChangedDelegate.AddDynamic(MyObj, &UAbilityTask_WaitMovementModeChange::OnMovementModeChange);
 		}
-
 		return MyObj;
 	}
 
