Fix SkillSystem in non-unity

diff --git a/Engine/Source/Runtime/SkillSystem/Classes/Abilities/GameplayAbility_CharacterJump.h b/Engine/Source/Runtime/SkillSystem/Classes/Abilities/GameplayAbility_CharacterJump.h
--- a/Engine/Source/Runtime/SkillSystem/Classes/Abilities/GameplayAbility_CharacterJump.h
+++ b/Engine/Source/Runtime/SkillSystem/Classes/Abilities/GameplayAbility_CharacterJump.h
@@ -2,6 +2,7 @@
 
 #pragma once
 
+#include "Abilities/GameplayAbility.h"
 #include "GameplayAbility_CharacterJump.generated.h"
 
 /**
diff --git a/Engine/Source/Runtime/SkillSystem/Classes/Abilities/GameplayAbility_Instanced.h b/Engine/Source/Runtime/SkillSystem/Classes/Abilities/GameplayAbility_Instanced.h
--- a/Engine/Source/Runtime/SkillSystem/Classes/Abilities/GameplayAbility_Instanced.h
+++ b/Engine/Source/Runtime/SkillSystem/Classes/Abilities/GameplayAbility_Instanced.h
@@ -1,7 +1,7 @@
 // Copyright 1998-2013 Epic Games, Inc. All Rights Reserved.
 
 #pragma once
-
+#include "GameplayAbility.h"
 #include "GameplayAbility_Instanced.generated.h"
 
 class UAnimInstance;
diff --git a/Engine/Source/Runtime/SkillSystem/Classes/Abilities/GameplayAbility_Montage.h b/Engine/Source/Runtime/SkillSystem/Classes/Abilities/GameplayAbility_Montage.h
--- a/Engine/Source/Runtime/SkillSystem/Classes/Abilities/GameplayAbility_Montage.h
+++ b/Engine/Source/Runtime/SkillSystem/Classes/Abilities/GameplayAbility_Montage.h
@@ -2,6 +2,7 @@
 
 #pragma once
 
+#include "Abilities/GameplayAbility.h"
 #include "GameplayAbility_Montage.generated.h"
 
 /**
diff --git a/Engine/Source/Runtime/SkillSystem/Classes/GameplayEffectExtension_LifestealTest.h b/Engine/Source/Runtime/SkillSystem/Classes/GameplayEffectExtension_LifestealTest.h
--- a/Engine/Source/Runtime/SkillSystem/Classes/GameplayEffectExtension_LifestealTest.h
+++ b/Engine/Source/Runtime/SkillSystem/Classes/GameplayEffectExtension_LifestealTest.h
@@ -2,6 +2,7 @@
 
 #pragma once
 
+#include "GameplayEffectExtension.h"
 #include "GameplayTagContainer.h"
 #include "GameplayTagAssetInterface.h"
 #include "AttributeSet.h"
diff --git a/Engine/Source/Runtime/SkillSystem/Classes/GameplayEffectExtension_ShieldTest.h b/Engine/Source/Runtime/SkillSystem/Classes/GameplayEffectExtension_ShieldTest.h
--- a/Engine/Source/Runtime/SkillSystem/Classes/GameplayEffectExtension_ShieldTest.h
+++ b/Engine/Source/Runtime/SkillSystem/Classes/GameplayEffectExtension_ShieldTest.h
@@ -2,6 +2,7 @@
 
 #pragma once
 
+#include "GameplayEffectExtension.h"
 #include "GameplayTagContainer.h"
 #include "GameplayTagAssetInterface.h"
 #include "AttributeSet.h"
diff --git a/Engine/Source/Runtime/SkillSystem/Classes/GameplayEffectStackingExtension.h b/Engine/Source/Runtime/SkillSystem/Classes/GameplayEffectStackingExtension.h
--- a/Engine/Source/Runtime/SkillSystem/Classes/GameplayEffectStackingExtension.h
+++ b/Engine/Source/Runtime/SkillSystem/Classes/GameplayEffectStackingExtension.h
@@ -5,6 +5,7 @@
 #include "GameplayTagContainer.h"
 #include "GameplayTagAssetInterface.h"
 #include "AttributeSet.h"
+#include "GameplayEffect.h"
 #include "GameplayEffectStackingExtension.generated.h"
 
 struct FGameplayEffectStackingCallbackData
diff --git a/Engine/Source/Runtime/SkillSystem/Classes/GameplayEffectStackingExtension_CappedNumberTest.h b/Engine/Source/Runtime/SkillSystem/Classes/GameplayEffectStackingExtension_CappedNumberTest.h
--- a/Engine/Source/Runtime/SkillSystem/Classes/GameplayEffectStackingExtension_CappedNumberTest.h
+++ b/Engine/Source/Runtime/SkillSystem/Classes/GameplayEffectStackingExtension_CappedNumberTest.h
@@ -2,6 +2,7 @@
 
 #pragma once
 
+#include "GameplayEffectStackingExtension.h"
 #include "GameplayTagContainer.h"
 #include "GameplayTagAssetInterface.h"
 #include "AttributeSet.h"
diff --git a/Engine/Source/Runtime/SkillSystem/Classes/GameplayEffectStackingExtension_DiminishingReturnsTest.h b/Engine/Source/Runtime/SkillSystem/Classes/GameplayEffectStackingExtension_DiminishingReturnsTest.h
--- a/Engine/Source/Runtime/SkillSystem/Classes/GameplayEffectStackingExtension_DiminishingReturnsTest.h
+++ b/Engine/Source/Runtime/SkillSystem/Classes/GameplayEffectStackingExtension_DiminishingReturnsTest.h
@@ -2,6 +2,7 @@
 
 #pragma once
 
+#include "GameplayEffectStackingExtension.h"
 #include "GameplayTagContainer.h"
 #include "GameplayTagAssetInterface.h"
 #include "AttributeSet.h"
diff --git a/Engine/Source/Runtime/SkillSystem/Classes/SkillSystemTestPawn.h b/Engine/Source/Runtime/SkillSystem/Classes/SkillSystemTestPawn.h
--- a/Engine/Source/Runtime/SkillSystem/Classes/SkillSystemTestPawn.h
+++ b/Engine/Source/Runtime/SkillSystem/Classes/SkillSystemTestPawn.h
@@ -1,6 +1,8 @@
 // Copyright 1998-2014 Epic Games, Inc. All Rights Reserved.
 
 #pragma once
+#include "GameplayCueInterface.h"
+#include "GameplayCueView.h"
 #include "SkillSystemTestPawn.generated.h"
 
 UCLASS(Blueprintable, BlueprintType)
diff --git a/Engine/Source/Runtime/SkillSystem/Private/Abilities/GameplayAbility.cpp b/Engine/Source/Runtime/SkillSystem/Private/Abilities/GameplayAbility.cpp
--- a/Engine/Source/Runtime/SkillSystem/Private/Abilities/GameplayAbility.cpp
+++ b/Engine/Source/Runtime/SkillSystem/Private/Abilities/GameplayAbility.cpp
@@ -1,6 +1,9 @@
 // Copyright 1998-2013 Epic Games, Inc. All Rights Reserved.
 
 #include "SkillSystemModulePrivatePCH.h"
+#include "Abilities/GameplayAbility.h"
+#include "GameplayEffect.h"
+#include "AttributeComponent.h"
 
 // --------------------------------------------------------------------------------------------------------------------------------------------------------
 //
diff --git a/Engine/Source/Runtime/SkillSystem/Private/Abilities/GameplayAbility_Instanced.cpp b/Engine/Source/Runtime/SkillSystem/Private/Abilities/GameplayAbility_Instanced.cpp
--- a/Engine/Source/Runtime/SkillSystem/Private/Abilities/GameplayAbility_Instanced.cpp
+++ b/Engine/Source/Runtime/SkillSystem/Private/Abilities/GameplayAbility_Instanced.cpp
@@ -1,6 +1,7 @@
 // Copyright 1998-2013 Epic Games, Inc. All Rights Reserved.
 
 #include "SkillSystemModulePrivatePCH.h"
+#include "Abilities/GameplayAbility_Instanced.h"
 
 // --------------------------------------------------------------------------------------------------------------------------------------------------------
 //
diff --git a/Engine/Source/Runtime/SkillSystem/Private/Abilities/GameplayAbility_Montage.cpp b/Engine/Source/Runtime/SkillSystem/Private/Abilities/GameplayAbility_Montage.cpp
--- a/Engine/Source/Runtime/SkillSystem/Private/Abilities/GameplayAbility_Montage.cpp
+++ b/Engine/Source/Runtime/SkillSystem/Private/Abilities/GameplayAbility_Montage.cpp
@@ -2,6 +2,8 @@
 
 #include "SkillSystemModulePrivatePCH.h"
 #include "Abilities/GameplayAbility_Montage.h"
+#include "GameplayEffect.h"
+#include "AttributeComponent.h"
 
 // --------------------------------------------------------------------------------------------------------------------------------------------------------
 //
diff --git a/Engine/Source/Runtime/SkillSystem/Private/Abilities/Tasks/AbilityTask_PlayMontageAndWait.cpp b/Engine/Source/Runtime/SkillSystem/Private/Abilities/Tasks/AbilityTask_PlayMontageAndWait.cpp
--- a/Engine/Source/Runtime/SkillSystem/Private/Abilities/Tasks/AbilityTask_PlayMontageAndWait.cpp
+++ b/Engine/Source/Runtime/SkillSystem/Private/Abilities/Tasks/AbilityTask_PlayMontageAndWait.cpp
@@ -1,5 +1,6 @@
 
 #include "SkillSystemModulePrivatePCH.h"
+#include "Abilities/Tasks/AbilityTask_PlayMontageAndWait.h"
 
 UAbilityTask_PlayMontageAndWait::UAbilityTask_PlayMontageAndWait(const class FPostConstructInitializeProperties& PCIP)
 	: Super(PCIP)
diff --git a/Engine/Source/Runtime/SkillSystem/Private/Abilities/Tasks/AbilityTask_WaitMovementModeChange.cpp b/Engine/Source/Runtime/SkillSystem/Private/Abilities/Tasks/AbilityTask_WaitMovementModeChange.cpp
--- a/Engine/Source/Runtime/SkillSystem/Private/Abilities/Tasks/AbilityTask_WaitMovementModeChange.cpp
+++ b/Engine/Source/Runtime/SkillSystem/Private/Abilities/Tasks/AbilityTask_WaitMovementModeChange.cpp
@@ -1,5 +1,6 @@
 
 #include "SkillSystemModulePrivatePCH.h"
+#include "Abilities/Tasks/AbilityTask_WaitMovementModeChange.h"
 
 UAbilityTask_WaitMovementModeChange::UAbilityTask_WaitMovementModeChange(const class FPostConstructInitializeProperties& PCIP)
 	: Super(PCIP)
diff --git a/Engine/Source/Runtime/SkillSystem/Private/AttributeComponent_Abilities.cpp b/Engine/Source/Runtime/SkillSystem/Private/AttributeComponent_Abilities.cpp
--- a/Engine/Source/Runtime/SkillSystem/Private/AttributeComponent_Abilities.cpp
+++ b/Engine/Source/Runtime/SkillSystem/Private/AttributeComponent_Abilities.cpp
@@ -2,6 +2,7 @@
 // ActorComponent.cpp: Actor component implementation.
 
 #include "SkillSystemModulePrivatePCH.h"
+#include "AttributeComponent.h"
 #include "Abilities/GameplayAbility.h"
 
 #include "Net/UnrealNetwork.h"
diff --git a/Engine/Source/Runtime/SkillSystem/Private/AttributeSet.cpp b/Engine/Source/Runtime/SkillSystem/Private/AttributeSet.cpp
--- a/Engine/Source/Runtime/SkillSystem/Private/AttributeSet.cpp
+++ b/Engine/Source/Runtime/SkillSystem/Private/AttributeSet.cpp
@@ -1,6 +1,7 @@
 // Copyright 1998-2013 Epic Games, Inc. All Rights Reserved.
 
 #include "SkillSystemModulePrivatePCH.h"
+#include "AttributeSet.h"
 
 void FGameplayAttribute::SetNumericValueChecked(const float NewValue, class UAttributeSet* Dest) const
 {
diff --git a/Engine/Source/Runtime/SkillSystem/Private/GameplayAbilitySet.cpp b/Engine/Source/Runtime/SkillSystem/Private/GameplayAbilitySet.cpp
--- a/Engine/Source/Runtime/SkillSystem/Private/GameplayAbilitySet.cpp
+++ b/Engine/Source/Runtime/SkillSystem/Private/GameplayAbilitySet.cpp
@@ -1,6 +1,8 @@
 // Copyright 1998-2014 Epic Games, Inc. All Rights Reserved.
 
 #include "SkillSystemModulePrivatePCH.h"
+#include "AttributeComponent.h"
+#include "Abilities/GameplayAbility.h"
 #include "GameplayAbilitySet.h"
 
 // --------------------------------------------------------------------------------------------------------------------------------------------------------
diff --git a/Engine/Source/Runtime/SkillSystem/Private/GameplayCueInterface.cpp b/Engine/Source/Runtime/SkillSystem/Private/GameplayCueInterface.cpp
--- a/Engine/Source/Runtime/SkillSystem/Private/GameplayCueInterface.cpp
+++ b/Engine/Source/Runtime/SkillSystem/Private/GameplayCueInterface.cpp
@@ -1,6 +1,7 @@
 // Copyright 1998-2014 Epic Games, Inc. All Rights Reserved.
 
 #include "SkillSystemModulePrivatePCH.h"
+#include "GameplayCueInterface.h"
 
 UGameplayCueInterface::UGameplayCueInterface(const FPostConstructInitializeProperties& PCIP)
 : Super(PCIP)
diff --git a/Engine/Source/Runtime/SkillSystem/Private/GameplayEffect.cpp b/Engine/Source/Runtime/SkillSystem/Private/GameplayEffect.cpp
--- a/Engine/Source/Runtime/SkillSystem/Private/GameplayEffect.cpp
+++ b/Engine/Source/Runtime/SkillSystem/Private/GameplayEffect.cpp
@@ -4,6 +4,7 @@
 #include "GameplayEffect.h"
 #include "GameplayEffectExtension.h"
 #include "GameplayEffectStackingExtension.h"
+#include "AttributeComponent.h"
 
 const float UGameplayEffect::INFINITE_DURATION = -1.f;
 const float UGameplayEffect::INSTANT_APPLICATION = 0.f;
diff --git a/Engine/Source/Runtime/SkillSystem/Private/GameplayEffectExtension.cpp b/Engine/Source/Runtime/SkillSystem/Private/GameplayEffectExtension.cpp
--- a/Engine/Source/Runtime/SkillSystem/Private/GameplayEffectExtension.cpp
+++ b/Engine/Source/Runtime/SkillSystem/Private/GameplayEffectExtension.cpp
@@ -1,6 +1,7 @@
 // Copyright 1998-2013 Epic Games, Inc. All Rights Reserved.
 
 #include "SkillSystemModulePrivatePCH.h"
+#include "GameplayEffectExtension.h"
 
 UGameplayEffectExtension::UGameplayEffectExtension(const class FPostConstructInitializeProperties& PCIP)
 : Super(PCIP)
diff --git a/Engine/Source/Runtime/SkillSystem/Private/GameplayEffectExtension_LifestealTest.cpp b/Engine/Source/Runtime/SkillSystem/Private/GameplayEffectExtension_LifestealTest.cpp
--- a/Engine/Source/Runtime/SkillSystem/Private/GameplayEffectExtension_LifestealTest.cpp
+++ b/Engine/Source/Runtime/SkillSystem/Private/GameplayEffectExtension_LifestealTest.cpp
@@ -4,6 +4,7 @@
 #include "GameplayEffectExtension_LifestealTest.h"
 #include "SkillSystemTestAttributeSet.h"
 #include "GameplayTagsModule.h"
+#include "AttributeComponent.h"
 
 UGameplayEffectExtension_LifestealTest::UGameplayEffectExtension_LifestealTest(const class FPostConstructInitializeProperties& PCIP)
 : Super(PCIP)
diff --git a/Engine/Source/Runtime/SkillSystem/Private/GameplayEffectExtension_ShieldTest.cpp b/Engine/Source/Runtime/SkillSystem/Private/GameplayEffectExtension_ShieldTest.cpp
--- a/Engine/Source/Runtime/SkillSystem/Private/GameplayEffectExtension_ShieldTest.cpp
+++ b/Engine/Source/Runtime/SkillSystem/Private/GameplayEffectExtension_ShieldTest.cpp
@@ -3,6 +3,8 @@
 #include "SkillSystemModulePrivatePCH.h"
 #include "GameplayEffectExtension_ShieldTest.h"
 #include "GameplayTagsModule.h"
+#include "AttributeComponent.h"
+#include "SkillSystemTestAttributeSet.h"
 
 UGameplayEffectExtension_ShieldTest::UGameplayEffectExtension_ShieldTest(const class FPostConstructInitializeProperties& PCIP)
 : Super(PCIP)
diff --git a/Engine/Source/Runtime/SkillSystem/Private/GameplayEffectStackingExtension.cpp b/Engine/Source/Runtime/SkillSystem/Private/GameplayEffectStackingExtension.cpp
--- a/Engine/Source/Runtime/SkillSystem/Private/GameplayEffectStackingExtension.cpp
+++ b/Engine/Source/Runtime/SkillSystem/Private/GameplayEffectStackingExtension.cpp
@@ -1,6 +1,7 @@
 // Copyright 1998-2014 Epic Games, Inc. All Rights Reserved.
 
 #include "SkillSystemModulePrivatePCH.h"
+#include "GameplayEffectStackingExtension.h"
 
 UGameplayEffectStackingExtension::UGameplayEffectStackingExtension(const class FPostConstructInitializeProperties& PCIP)
 : Super(PCIP)
diff --git a/Engine/Source/Runtime/SkillSystem/Private/GameplayEffectTests.cpp b/Engine/Source/Runtime/SkillSystem/Private/GameplayEffectTests.cpp
--- a/Engine/Source/Runtime/SkillSystem/Private/GameplayEffectTests.cpp
+++ b/Engine/Source/Runtime/SkillSystem/Private/GameplayEffectTests.cpp
@@ -2,9 +2,16 @@
 
 #include "SkillSystemModulePrivatePCH.h"
 #include "SkillSystemTestPawn.h"
+#include "SkillSystemTestAttributeSet.h"
 #include "GameplayEffect.h"
 #include "AttributeSet.h"
 #include "GameplayTagsModule.h"
+#include "AttributeComponent.h"
+#include "GameplayEffectExtension_LifestealTest.h"
+#include "GameplayEffectExtension_ShieldTest.h"
+#include "GameplayEffectStackingExtension_CappedNumberTest.h"
+#include "GameplayEffectStackingExtension_DiminishingReturnsTest.h"
+
 //#include "SkillSystemTestPawn.h"
 
 IMPLEMENT_SIMPLE_AUTOMATION_TEST(FGameplayEffectsTest, "SkillSystem.GameplayEffects", EAutomationTestFlags::ATF_Editor)
diff --git a/Engine/Source/Runtime/SkillSystem/Private/SkillSystemDebugHUD.cpp b/Engine/Source/Runtime/SkillSystem/Private/SkillSystemDebugHUD.cpp
--- a/Engine/Source/Runtime/SkillSystem/Private/SkillSystemDebugHUD.cpp
+++ b/Engine/Source/Runtime/SkillSystem/Private/SkillSystemDebugHUD.cpp
@@ -1,6 +1,7 @@
 // Copyright 1998-2014 Epic Games, Inc. All Rights Reserved.
 
 #include "SkillSystemModulePrivatePCH.h"
+#include "AttributeComponent.h"
 #include "SkillSystemDebugHUD.h"
 #include "DebugRenderSceneProxy.h"
 
diff --git a/Engine/Source/Runtime/SkillSystem/Private/SkillSystemTestAttributeSet.cpp b/Engine/Source/Runtime/SkillSystem/Private/SkillSystemTestAttributeSet.cpp
--- a/Engine/Source/Runtime/SkillSystem/Private/SkillSystemTestAttributeSet.cpp
+++ b/Engine/Source/Runtime/SkillSystem/Private/SkillSystemTestAttributeSet.cpp
@@ -1,7 +1,11 @@
 // Copyright 1998-2014 Epic Games, Inc. All Rights Reserved.
 
 #include "SkillSystemModulePrivatePCH.h"
+#include "SkillSystemTestAttributeSet.h"
 #include "GameplayTagsModule.h"
+#include "GameplayEffectExtension.h"
+#include "AttributeComponent.h"
+
 
 USkillSystemTestAttributeSet::USkillSystemTestAttributeSet(const class FPostConstructInitializeProperties& PCIP)
 : Super(PCIP)
diff --git a/Engine/Source/Runtime/SkillSystem/Private/SkillSystemTestPawn.cpp b/Engine/Source/Runtime/SkillSystem/Private/SkillSystemTestPawn.cpp
--- a/Engine/Source/Runtime/SkillSystem/Private/SkillSystemTestPawn.cpp
+++ b/Engine/Source/Runtime/SkillSystem/Private/SkillSystemTestPawn.cpp
@@ -1,6 +1,9 @@
 // Copyright 1998-2014 Epic Games, Inc. All Rights Reserved.
 
 #include "SkillSystemModulePrivatePCH.h"
+#include "SkillSystemTestPawn.h"
+#include "SkillSystemTestAttributeSet.h"
+#include "AttributeComponent.h"
 
 FName  ASkillSystemTestPawn::AttributeComponentName(TEXT("AttributeComponent0"));
 
