Adding tags used in SkillSystemTestAttributeSet to the tag table for QAGame

diff --git a/Engine/Source/Runtime/SkillSystem/Private/SkillSystemTestAttributeSet.cpp b/Engine/Source/Runtime/SkillSystem/Private/SkillSystemTestAttributeSet.cpp
--- a/Engine/Source/Runtime/SkillSystem/Private/SkillSystemTestAttributeSet.cpp
+++ b/Engine/Source/Runtime/SkillSystem/Private/SkillSystemTestAttributeSet.cpp
@@ -73,7 +73,7 @@ void USkillSystemTestAttributeSet::PreAttributeModify(struct FGameplayEffectModC
 				// This is a trivial/naive implementation of armor. It assumes the rmorDamageReduction is an actual % to reduce physics damage by.
 				// Real games would probably use armor rating as an attribute and calculate a % here based on the damage source's level, etc.
 				Data.EvaluatedData.Magnitude *= (1.f - ArmorDamageReduction);
-				Data.EvaluatedData.Tags.AddTag(IGameplayTagsModule::RequestGameplayTag(FName(TEXT("Damage.Mitigatd.Armor"))));
+				Data.EvaluatedData.Tags.AddTag(IGameplayTagsModule::RequestGameplayTag(FName(TEXT("Damage.Mitigated.Armor"))));
 			}
 		}
 
