Abilty system updates

diff --git a/Engine/Source/Editor/SkillSystemEditor/Classes/K2Node_WaitOverlap.h b/Engine/Source/Editor/SkillSystemEditor/Classes/K2Node_WaitOverlap.h
new file mode 100644
--- /dev/null
+++ b/Engine/Source/Editor/SkillSystemEditor/Classes/K2Node_WaitOverlap.h
@@ -0,0 +1,17 @@
+// Copyright 1998-2014 Epic Games, Inc. All Rights Reserved.
+
+#pragma once
+#include "K2Node_WaitOverlap.generated.h"
+
+UCLASS()
+class UK2Node_WaitOverlap : public UK2Node_BaseAsyncTask
+{
+	GENERATED_UCLASS_BODY()
+
+	// Begin UEdGraphNode interface
+	virtual FString GetTooltip() const OVERRIDE;
+	virtual FText GetNodeTitle(ENodeTitleType::Type TitleType) const OVERRIDE;
+	// End UEdGraphNode interface
+
+	virtual FString GetCategoryName();
+};
diff --git a/Engine/Source/Editor/SkillSystemEditor/Private/K2Node_PlayMontageAndWait.cpp b/Engine/Source/Editor/SkillSystemEditor/Private/K2Node_PlayMontageAndWait.cpp
--- a/Engine/Source/Editor/SkillSystemEditor/Private/K2Node_PlayMontageAndWait.cpp
+++ b/Engine/Source/Editor/SkillSystemEditor/Private/K2Node_PlayMontageAndWait.cpp
@@ -2,6 +2,7 @@
 
 
 #include "SkillSystemEditorModulePrivatePCH.h"
+#include "K2Node_PlayMontageAndWait.h"
 #include "SkillSystemBlueprintLibrary.h"
 #include "Abilities/Tasks/AbilityTask_PlayMontageAndWait.h"
 #include "CompilerResultsLog.h"
diff --git a/Engine/Source/Editor/SkillSystemEditor/Private/K2Node_WaitMovementModeChange.cpp b/Engine/Source/Editor/SkillSystemEditor/Private/K2Node_WaitMovementModeChange.cpp
--- a/Engine/Source/Editor/SkillSystemEditor/Private/K2Node_WaitMovementModeChange.cpp
+++ b/Engine/Source/Editor/SkillSystemEditor/Private/K2Node_WaitMovementModeChange.cpp
@@ -2,6 +2,7 @@
 
 
 #include "SkillSystemEditorModulePrivatePCH.h"
+#include "K2Node_WaitMovementModeChange.h"
 #include "SkillSystemBlueprintLibrary.h"
 #include "Abilities/Tasks/AbilityTask_WaitMovementModeChange.h"
 #include "CompilerResultsLog.h"
diff --git a/Engine/Source/Editor/SkillSystemEditor/Private/K2Node_WaitOverlap.cpp b/Engine/Source/Editor/SkillSystemEditor/Private/K2Node_WaitOverlap.cpp
new file mode 100644
--- /dev/null
+++ b/Engine/Source/Editor/SkillSystemEditor/Private/K2Node_WaitOverlap.cpp
@@ -0,0 +1,38 @@
+// Copyright 1998-2014 Epic Games, Inc. All Rights Reserved.
+
+
+#include "SkillSystemEditorModulePrivatePCH.h"
+#include "K2Node_WaitOverlap.h"
+#include "SkillSystemBlueprintLibrary.h"
+#include "Abilities/Tasks/AbilityTask_WaitOverlap.h"
+#include "CompilerResultsLog.h"
+
+#define LOCTEXT_NAMESPACE "K2Node_WaitOverlap"
+
+UK2Node_WaitOverlap::UK2Node_WaitOverlap(const class FPostConstructInitializeProperties& PCIP)
+	: Super(PCIP)
+{
+	
+	ProxyFactoryFunctionName = GET_FUNCTION_NAME_CHECKED(USkillSystemBlueprintLibrary, CreateWaitOverlap);
+	ProxyFactoryClass = USkillSystemBlueprintLibrary::StaticClass();
+	ProxyClass = UAbilityTask_WaitOverlap::StaticClass();
+}
+
+FString UK2Node_WaitOverlap::GetCategoryName()
+{
+	return TEXT("Ability");
+}
+
+FString UK2Node_WaitOverlap::GetTooltip() const
+{
+	return TEXT("Wait until an overlap occurs");
+}
+
+FText UK2Node_WaitOverlap::GetNodeTitle(ENodeTitleType::Type TitleType) const
+{
+	return LOCTEXT("WaitForOverlap", "Wait For Overlap");
+}
+
+#undef LOCTEXT_NAMESPACE
+
+
diff --git a/Engine/Source/Editor/SkillSystemEditor/Private/SkillSystemEditorModulePrivatePCH.h b/Engine/Source/Editor/SkillSystemEditor/Private/SkillSystemEditorModulePrivatePCH.h
--- a/Engine/Source/Editor/SkillSystemEditor/Private/SkillSystemEditorModulePrivatePCH.h
+++ b/Engine/Source/Editor/SkillSystemEditor/Private/SkillSystemEditorModulePrivatePCH.h
@@ -12,6 +12,4 @@
 
 #include "Core.h"
 #include "Engine.h"
-#include "BlueprintGraphDefinitions.h"
-
-#include "SkillSystemEditorClasses.h"
\ No newline at end of file
+#include "BlueprintGraphDefinitions.h"
\ No newline at end of file
diff --git a/Engine/Source/Runtime/SkillSystem/Classes/Abilities/GameplayAbility.h b/Engine/Source/Runtime/SkillSystem/Classes/Abilities/GameplayAbility.h
--- a/Engine/Source/Runtime/SkillSystem/Classes/Abilities/GameplayAbility.h
+++ b/Engine/Source/Runtime/SkillSystem/Classes/Abilities/GameplayAbility.h
@@ -153,6 +153,8 @@ public:
 	//		Commit()				- Commits reources/cooldowns etc. Activate() must call this!
 	//		
 	//		CancelAbility()			- 
+	//
+	//		EndAbility()			- 
 	//	
 	// ----------------------------------------------------------------------------------------------------------------
 
@@ -176,6 +178,10 @@ public:
 	/** Input binding. Base implementation does nothing */
 	virtual void InputReleased(int32 InputID, const FGameplayAbilityActorInfo ActorInfo);
 
+	UFUNCTION(BlueprintCallable, Category=Ability)
+	virtual void EndAbility(const FGameplayAbilityActorInfo ActorInfo);
+
+
 	// ----------------------------------------------------------------------------------------------------------------
 	//
 	//	Outward facing API for other systems
diff --git a/Engine/Source/Runtime/SkillSystem/Classes/Abilities/GameplayAbility_Instanced.h b/Engine/Source/Runtime/SkillSystem/Classes/Abilities/GameplayAbility_Instanced.h
--- a/Engine/Source/Runtime/SkillSystem/Classes/Abilities/GameplayAbility_Instanced.h
+++ b/Engine/Source/Runtime/SkillSystem/Classes/Abilities/GameplayAbility_Instanced.h
@@ -18,6 +18,8 @@ class SKILLSYSTEM_API UGameplayAbility_Instanced : public UGameplayAbility
 
 public:
 	
+	virtual void EndAbility(const FGameplayAbilityActorInfo ActorInfo) OVERRIDE;
+	
 	// -------------------------------------------------
 
 	UFUNCTION(BlueprintImplementableEvent, Category = Ability)
diff --git a/Engine/Source/Runtime/SkillSystem/Classes/Abilities/Tasks/AbilityTask_WaitOverlap.h b/Engine/Source/Runtime/SkillSystem/Classes/Abilities/Tasks/AbilityTask_WaitOverlap.h
new file mode 100644
--- /dev/null
+++ b/Engine/Source/Runtime/SkillSystem/Classes/Abilities/Tasks/AbilityTask_WaitOverlap.h
@@ -0,0 +1,28 @@
+// Copyright 1998-2014 Epic Games, Inc. All Rights Reserved.
+#pragma once
+#include "AbilityTask.h"
+#include "SkillSystemTypes.h"
+#include "AbilityTask_WaitOverlap.generated.h"
+
+DECLARE_DYNAMIC_MULTICAST_DELEGATE_OneParam(FWaitOverlapDelegate, FGameplayAbilityTargetDataHandle, TargetData);
+
+class AActor;
+class UPrimitiveComponent;
+
+UCLASS(MinimalAPI)
+class UAbilityTask_WaitOverlap : public UAbilityTask
+{
+	GENERATED_UCLASS_BODY()
+
+	UPROPERTY(BlueprintAssignable)
+	FWaitOverlapDelegate	OnOverlap;
+
+	UFUNCTION()
+	void OnOverlapCallback(AActor* OtherActor, UPrimitiveComponent* OtherComp, int32 OtherBodyIndex, bool bFromSweep, const FHitResult & SweepResult);
+
+	UFUNCTION()
+	void OnHitCallback(AActor* OtherActor, UPrimitiveComponent* OtherComp, FVector NormalImpulse, const FHitResult& Hit);
+
+public:
+	
+};
\ No newline at end of file
diff --git a/Engine/Source/Runtime/SkillSystem/Classes/GameplayCueInterface.h b/Engine/Source/Runtime/SkillSystem/Classes/GameplayCueInterface.h
--- a/Engine/Source/Runtime/SkillSystem/Classes/GameplayCueInterface.h
+++ b/Engine/Source/Runtime/SkillSystem/Classes/GameplayCueInterface.h
@@ -3,6 +3,7 @@
 #pragma once
 
 #include "GameplayTagContainer.h"
+#include "GameplayEffect.h"
 #include "GameplayCueInterface.generated.h"
 
 /** Interface for actors that wish to handle GameplayCue events from GameplayEffects */
@@ -18,16 +19,16 @@ class SKILLSYSTEM_API IGameplayCueInterface
 
 	// EVENT
 	UFUNCTION()
-	virtual void GameplayCueActivated(const FGameplayTagContainer & GameplayCueTags, float NormalizedMagnitude) = 0;
+	virtual void GameplayCueActivated(const FGameplayTagContainer & GameplayCueTags, float NormalizedMagnitude, const FGameplayEffectInstigatorContext InstigatorContext) = 0;
 
 	UFUNCTION()
-	virtual void GameplayCueExecuted(const FGameplayTagContainer & GameplayCueTags, float NormalizedMagnitude) = 0;
+	virtual void GameplayCueExecuted(const FGameplayTagContainer & GameplayCueTags, float NormalizedMagnitude, const FGameplayEffectInstigatorContext InstigatorContext) = 0;
 
 
 	// STATE
 	UFUNCTION()
-	virtual void GameplayCueAdded(const FGameplayTagContainer & GameplayCueTags, float NormalizedMagnitude) = 0;
+	virtual void GameplayCueAdded(const FGameplayTagContainer & GameplayCueTags, float NormalizedMagnitude, const FGameplayEffectInstigatorContext InstigatorContext) = 0;
 
 	UFUNCTION()
-	virtual void GameplayCueRemoved(const FGameplayTagContainer & GameplayCueTags, float NormalizedMagnitude) = 0;
+	virtual void GameplayCueRemoved(const FGameplayTagContainer & GameplayCueTags, float NormalizedMagnitude, const FGameplayEffectInstigatorContext InstigatorContext) = 0;
 };
\ No newline at end of file
diff --git a/Engine/Source/Runtime/SkillSystem/Classes/GameplayCueView.h b/Engine/Source/Runtime/SkillSystem/Classes/GameplayCueView.h
--- a/Engine/Source/Runtime/SkillSystem/Classes/GameplayCueView.h
+++ b/Engine/Source/Runtime/SkillSystem/Classes/GameplayCueView.h
@@ -75,7 +75,7 @@ struct FGameplayCueViewInfo
 	UPROPERTY(EditAnywhere, BlueprintReadWrite, Category = GameplayCue)
 	class TSubclassOf<AActor> ActorClass;
 
-	virtual TSharedPtr<FGameplayCueViewEffects> SpawnViewEffects(AActor *Owner, TArray<UObject*> *SpawnedObjects) const;
+	virtual TSharedPtr<FGameplayCueViewEffects> SpawnViewEffects(AActor *Owner, TArray<UObject*> *SpawnedObjects, const FGameplayEffectInstigatorContext InstigatorContext) const;
 };
 
 UCLASS(BlueprintType)
@@ -109,16 +109,18 @@ struct SKILLSYSTEM_API FGameplayCueHandler
 	UPROPERTY()
 	TArray<UObject*> SpawnedObjects;
 
-	virtual void GameplayCueActivated(const FGameplayTagContainer & GameplayCueTags, float NormalizedMagnitude);
+	virtual void GameplayCueActivated(const FGameplayTagContainer & GameplayCueTags, float NormalizedMagnitude, const FGameplayEffectInstigatorContext InstigatorContext);
 	
-	virtual void GameplayCueExecuted(const FGameplayTagContainer & GameplayCueTags, float NormalizedMagnitude);
+	virtual void GameplayCueExecuted(const FGameplayTagContainer & GameplayCueTags, float NormalizedMagnitude, const FGameplayEffectInstigatorContext InstigatorContext);
 	
-	virtual void GameplayCueAdded(const FGameplayTagContainer & GameplayCueTags, float NormalizedMagnitude);
+	virtual void GameplayCueAdded(const FGameplayTagContainer & GameplayCueTags, float NormalizedMagnitude, const FGameplayEffectInstigatorContext InstigatorContext);
 	
-	virtual void GameplayCueRemoved(const FGameplayTagContainer & GameplayCueTags, float NormalizedMagnitude);
+	virtual void GameplayCueRemoved(const FGameplayTagContainer & GameplayCueTags, float NormalizedMagnitude, const FGameplayEffectInstigatorContext InstigatorContext);
 
 private:
 
+	FGameplayCueViewInfo * GetBestMatchingView(EGameplayCueEvent::Type Type, const FGameplayTag BaseTag);
+
 	void ClearEffects(TArray< TSharedPtr<FGameplayCueViewEffects > > &Effects);
 
 
diff --git a/Engine/Source/Runtime/SkillSystem/Classes/GameplayEffect.h b/Engine/Source/Runtime/SkillSystem/Classes/GameplayEffect.h
--- a/Engine/Source/Runtime/SkillSystem/Classes/GameplayEffect.h
+++ b/Engine/Source/Runtime/SkillSystem/Classes/GameplayEffect.h
@@ -23,6 +23,7 @@ DECLARE_MULTICAST_DELEGATE_ThreeParams(FOnAttributeGameplayEffectSpecExected, co
 #endif
 
 class UGameplayEffect;
+class UAttributeComponent;
 
 UENUM(BlueprintType)
 namespace EGameplayModOp
@@ -601,6 +602,7 @@ struct FGameplayEffectInstigatorContext
 
 	FGameplayEffectInstigatorContext()
 		: Instigator(NULL)
+		, InstigatorAttributeComponent(NULL)
 	{
 	}
 
@@ -615,6 +617,8 @@ struct FGameplayEffectInstigatorContext
 
 	void AddInstigator(class AActor *InInstigator);
 
+	void AddHitResult(const FHitResult InHitResult);
+
 	FString ToString() const
 	{
 		return Instigator ? Instigator->GetName() : FString(TEXT("NONE"));
@@ -626,20 +630,35 @@ struct FGameplayEffectInstigatorContext
 		return Instigator;
 	}
 
-	class UAttributeComponent * GetOriginInstigatorAttributeComponent() const
+	UAttributeComponent * GetOriginalInstigatorAttributeComponent() const
 	{
 		return InstigatorAttributeComponent;
 	}
-
-	// @todo: Need a way to track instigator per stack
+	
 	/** Instigator controller */
-	UPROPERTY(NotReplicated)
+	UPROPERTY()
 	class AActor* Instigator;
 
 	UPROPERTY(NotReplicated)
-	class UAttributeComponent *InstigatorAttributeComponent;
+	UAttributeComponent *InstigatorAttributeComponent;
+
+	/** Trace information - may be NULL in many cases */
+	TSharedPtr<FHitResult>	HitResult;
+
+	bool NetSerialize(FArchive& Ar, class UPackageMap* Map, bool& bOutSuccess);
+};
+
+template<>
+struct TStructOpsTypeTraits< FGameplayEffectInstigatorContext > : public TStructOpsTypeTraitsBase
+{
+	enum
+	{
+		WithNetSerializer = true,
+		WithCopy = true		// Necessary so that TSharedPtr<FGameplayAbilityTargetData> Data is copied around
+	};
 };
 
+
 /**
  * FAggregatorRef
  *	A reference to an FAggregator. The reference may be weak or hard, and this can be changed over the lifetime of the FAggregatorRef.
@@ -1124,7 +1143,7 @@ struct FGameplayEffectSpec
 		// If we initialize a GameplayEffectSpec with no level object passed in.
 	}
 
-	FGameplayEffectSpec( const UGameplayEffect *InDef, AActor *Owner, AActor *Instigator, float Level, const FGlobalCurveDataOverride *CurveData );
+	FGameplayEffectSpec( const UGameplayEffect *InDef, AActor *Instigator, float Level, const FGlobalCurveDataOverride *CurveData );
 	
 	UPROPERTY()
 	const UGameplayEffect * Def;
@@ -1133,7 +1152,8 @@ struct FGameplayEffectSpec
 	TSharedPtr< FGameplayEffectLevelSpec > ModifierLevel;
 	
 	// Replicated
-	FGameplayEffectInstigatorContext InstigatorStack; // This tells us how we got here (who / what applied us)
+	UPROPERTY()
+	FGameplayEffectInstigatorContext InstigatorContext; // This tells us how we got here (who / what applied us)
 
 	float GetDuration() const;
 	float GetPeriod() const;
@@ -1356,7 +1376,7 @@ struct FActiveGameplayEffectsContainer : public FFastArraySerializer
 
 	void OnPropertyAggregatorDirty(FAggregator* Aggregator, FGameplayAttribute Attribute);
 
-	class UAttributeComponent *Owner;
+	UAttributeComponent *Owner;
 
 	void TEMP_TickActiveEffects(float DeltaSeconds);
 
diff --git a/Engine/Source/Runtime/SkillSystem/Classes/SkillSystemBlueprintLibrary.h b/Engine/Source/Runtime/SkillSystem/Classes/SkillSystemBlueprintLibrary.h
--- a/Engine/Source/Runtime/SkillSystem/Classes/SkillSystemBlueprintLibrary.h
+++ b/Engine/Source/Runtime/SkillSystem/Classes/SkillSystemBlueprintLibrary.h
@@ -3,22 +3,33 @@
 #pragma once
 
 #include "Abilities/GameplayAbility.h"
+#include "SkillSystemTypes.h"
 #include "SkillSystemBlueprintLibrary.generated.h"
 
-class UBlueprintPlayMontageAndWaitTaskProxy;
+class UAnimMontage;
+class UAbilityTask_PlayMontageAndWait;
+class UAbilityTask_WaitMovementModeChange;
+class UAbilityTask_WaitOverlap;
 
 UCLASS(MinimalAPI)
 class USkillSystemBlueprintLibrary : public UBlueprintFunctionLibrary
 {
 	GENERATED_UCLASS_BODY()
 
-	UFUNCTION(BlueprintCallable, Category = "Utilities|FlowControl", meta = (Latent, HidePin = "WorldContextObject", DefaultToSelf = "WorldContextObject", LatentInfo = "LatentInfo" ))
-	static void	StaticFuncPlayMontageAndWait(AActor* WorldContextObject, UAnimMontage * Montage, struct FLatentActionInfo LatentInfo);
+	UFUNCTION(BlueprintPure, Category = Ability)
+	static UAttributeComponent* GetAttributeComponent(AActor *Actor);
 
+	UFUNCTION(BlueprintCallable, Category=Ability)
+	static void ApplyGameplayEffectToTargetData(FGameplayAbilityTargetDataHandle Target, UGameplayEffect *GameplayEffect, const FGameplayAbilityActorInfo InstigatorInfo);
+
+	// -------------------------------------------------------------------------------
 
 	UFUNCTION(BlueprintCallable, meta=(HidePin="WorldContextObject", DefaultToSelf="WorldContextObject", BlueprintInternalUseOnly = "TRUE"))
-	static class UAbilityTask_PlayMontageAndWait* CreatePlayMontageAndWaitProxy(class UObject* WorldContextObject, class UAnimMontage *MontageToPlay);
+	static UAbilityTask_PlayMontageAndWait* CreatePlayMontageAndWaitProxy(UObject* WorldContextObject, UAnimMontage *MontageToPlay);
+
+	UFUNCTION(BlueprintCallable, meta = (HidePin = "WorldContextObject", DefaultToSelf = "WorldContextObject", BlueprintInternalUseOnly = "TRUE"))
+	static UAbilityTask_WaitMovementModeChange* CreateWaitMovementModeChange(UObject* WorldContextObject, EMovementMode NewMode);
 
 	UFUNCTION(BlueprintCallable, meta = (HidePin = "WorldContextObject", DefaultToSelf = "WorldContextObject", BlueprintInternalUseOnly = "TRUE"))
-	static class UAbilityTask_WaitMovementModeChange* CreateWaitMovementModeChange(class UObject* WorldContextObject, EMovementMode NewMode);
+	static UAbilityTask_WaitOverlap* CreateWaitOverlap(UObject* WorldContextObject, EMovementMode NewMode);
 };
\ No newline at end of file
diff --git a/Engine/Source/Runtime/SkillSystem/Classes/SkillSystemGlobals.h b/Engine/Source/Runtime/SkillSystem/Classes/SkillSystemGlobals.h
--- a/Engine/Source/Runtime/SkillSystem/Classes/SkillSystemGlobals.h
+++ b/Engine/Source/Runtime/SkillSystem/Classes/SkillSystemGlobals.h
@@ -4,6 +4,9 @@
 
 #include "SkillSystemGlobals.generated.h"
 
+class AActor;
+class UAttributeComponent;
+
 /** Holds global data for the skill system. Can be configured per project via config file */
 UCLASS(config=Game)
 class SKILLSYSTEM_API USkillSystemGlobals : public UObject
@@ -32,6 +35,12 @@ class SKILLSYSTEM_API USkillSystemGlobals : public UObject
 		GlobalAttributeDataTable = InTable;
 	}
 
+	/**
+	 *	Games may want to override this in a UMyProjectSkillSystemsGlobals class and provide
+	 *	a faster lookup for attribute compnents (E.g. Cast<MyProjPawn>(Actor)->AttributeComponent; etc)
+	 */
+	virtual UAttributeComponent * GetAttributeComponentFromActor(AActor *Actor) const;
+
 private:
 
 	class UCurveTable* GlobalCurveTable;
diff --git a/Engine/Source/Runtime/SkillSystem/Classes/SkillSystemTestPawn.h b/Engine/Source/Runtime/SkillSystem/Classes/SkillSystemTestPawn.h
--- a/Engine/Source/Runtime/SkillSystem/Classes/SkillSystemTestPawn.h
+++ b/Engine/Source/Runtime/SkillSystem/Classes/SkillSystemTestPawn.h
@@ -13,16 +13,16 @@ class SKILLSYSTEM_API ASkillSystemTestPawn : public ADefaultPawn, public IGamepl
 	virtual void PostInitializeComponents() OVERRIDE;
 
 	UFUNCTION()
-	virtual void GameplayCueActivated(const FGameplayTagContainer & GameplayCueTags, float NormalizedMagnitude) OVERRIDE;
+	virtual void GameplayCueActivated(const FGameplayTagContainer & GameplayCueTags, float NormalizedMagnitude, const FGameplayEffectInstigatorContext InstigatorContext) OVERRIDE;
 
 	UFUNCTION()
-	virtual void GameplayCueExecuted(const FGameplayTagContainer & GameplayCueTags, float NormalizedMagnitude) OVERRIDE;
+	virtual void GameplayCueExecuted(const FGameplayTagContainer & GameplayCueTags, float NormalizedMagnitude, const FGameplayEffectInstigatorContext InstigatorContext) OVERRIDE;
 
 	UFUNCTION()
-	virtual void GameplayCueAdded(const FGameplayTagContainer & GameplayCueTags, float NormalizedMagnitude) OVERRIDE;
+	virtual void GameplayCueAdded(const FGameplayTagContainer & GameplayCueTags, float NormalizedMagnitude, const FGameplayEffectInstigatorContext InstigatorContext) OVERRIDE;
 
 	UFUNCTION()
-	virtual void GameplayCueRemoved(const FGameplayTagContainer & GameplayCueTags, float NormalizedMagnitude) OVERRIDE;
+	virtual void GameplayCueRemoved(const FGameplayTagContainer & GameplayCueTags, float NormalizedMagnitude, const FGameplayEffectInstigatorContext InstigatorContext) OVERRIDE;
 
 
 	/** DefaultPawn collision component */
diff --git a/Engine/Source/Runtime/SkillSystem/Private/Abilities/GameplayAbility.cpp b/Engine/Source/Runtime/SkillSystem/Private/Abilities/GameplayAbility.cpp
--- a/Engine/Source/Runtime/SkillSystem/Private/Abilities/GameplayAbility.cpp
+++ b/Engine/Source/Runtime/SkillSystem/Private/Abilities/GameplayAbility.cpp
@@ -150,6 +150,11 @@ bool UGameplayAbility::TryActivateAbility(const FGameplayAbilityActorInfo ActorI
 	return true;	
 }
 
+void UGameplayAbility::EndAbility(const FGameplayAbilityActorInfo ActorInfo)
+{
+
+}
+
 void UGameplayAbility::ActivateAbility(FGameplayAbilityActorInfo ActorInfo)
 {
 	// Child classes may want to do stuff here...
@@ -331,4 +336,38 @@ void FGameplayAbilityActorInfo::SetActivationConfirmed()
 	ActivationMode = EGameplayAbilityActivationMode::Confirmed;
 }
 
-//----------------------------------------------------------------------
\ No newline at end of file
+//----------------------------------------------------------------------
+
+void FGameplayAbilityTargetData::ApplyGameplayEffect(UGameplayEffect *GameplayEffect, const FGameplayAbilityActorInfo InstigatorInfo)
+{
+	// Improve relationship between InstigatorContext and 
+	
+
+	FGameplayEffectSpec	SpecToApply(GameplayEffect,					// The UGameplayEffect data asset
+									InstigatorInfo.Actor.Get(),		// The actor who instigated this
+									1.f,							// FIXME: Leveling
+									NULL							// FIXME: CurveData override... should we just remove this?
+									);
+	if (HasHitResult())
+	{
+		SpecToApply.InstigatorContext.AddHitResult(*GetHitResult());
+	}
+
+	TArray<AActor*> Actors = GetActors();
+	for (AActor * TargetActor : Actors)
+	{
+		check(TargetActor);
+		
+		UAttributeComponent * TargetComponent = USkillSystemBlueprintLibrary::GetAttributeComponent(TargetActor);
+		if (TargetComponent)
+		{
+			InstigatorInfo.AttributeComponent->ApplyGameplayEffectSpecToTarget(SpecToApply, TargetComponent);
+		}
+	}
+}
+
+
+FString FGameplayAbilityTargetData::ToString() const
+{
+	return TEXT("BASE CLASS");
+}
diff --git a/Engine/Source/Runtime/SkillSystem/Private/Abilities/GameplayAbility_Instanced.cpp b/Engine/Source/Runtime/SkillSystem/Private/Abilities/GameplayAbility_Instanced.cpp
--- a/Engine/Source/Runtime/SkillSystem/Private/Abilities/GameplayAbility_Instanced.cpp
+++ b/Engine/Source/Runtime/SkillSystem/Private/Abilities/GameplayAbility_Instanced.cpp
@@ -57,8 +57,6 @@ void UGameplayAbility_Instanced::CommitExecute(const FGameplayAbilityActorInfo A
 
 void UGameplayAbility_Instanced::ActivateAbility(FGameplayAbilityActorInfo ActorInfo)
 {
-	//check(!HasAnyFlags(RF_ClassDefaultObject));
-
 	// Call into blueprint (fixme: cache this)
 	static FName FuncName = FName(TEXT("K2_ActivateAbility"));
 	UFunction* CanActivateFunction = GetClass()->FindFunctionByName(FuncName);
@@ -111,4 +109,17 @@ bool UGameplayAbility_Instanced::CallRemoteFunction(UFunction* Function, void* P
 	}
 
 	return false;
+}
+
+void UGameplayAbility_Instanced::EndAbility(const FGameplayAbilityActorInfo ActorInfo)
+{
+	Super::EndAbility(ActorInfo);
+
+	if (InstancedPerExecution)
+	{
+		MarkPendingKill();
+	}
+
+	// Remove from owning attributecomponent?
+	// generic way of releasing all callbacks!
 }
\ No newline at end of file
diff --git a/Engine/Source/Runtime/SkillSystem/Private/Abilities/Tasks/AbilityTask_WaitOverlap.cpp b/Engine/Source/Runtime/SkillSystem/Private/Abilities/Tasks/AbilityTask_WaitOverlap.cpp
new file mode 100644
--- /dev/null
+++ b/Engine/Source/Runtime/SkillSystem/Private/Abilities/Tasks/AbilityTask_WaitOverlap.cpp
@@ -0,0 +1,36 @@
+
+#include "SkillSystemModulePrivatePCH.h"
+#include "Abilities/Tasks/AbilityTask_WaitOverlap.h"
+
+UAbilityTask_WaitOverlap::UAbilityTask_WaitOverlap(const class FPostConstructInitializeProperties& PCIP)
+	: Super(PCIP)
+{
+
+}
+
+
+void UAbilityTask_WaitOverlap::OnOverlapCallback(class AActor* OtherActor, class UPrimitiveComponent* OtherComp, int32 OtherBodyIndex, bool bFromSweep, const FHitResult & SweepResult)
+{
+	if (OtherActor && bFromSweep)
+	{	
+		//OnOverlap.Broadcast(SweepResult);
+	}
+}
+
+void UAbilityTask_WaitOverlap::OnHitCallback(AActor* OtherActor, UPrimitiveComponent* OtherComp, FVector NormalImpulse, const FHitResult& Hit)
+{
+	if(OtherActor)
+	{
+		// Construct TargetData
+		FGameplayAbilityTargetData_SingleTargetHit * TargetData = new FGameplayAbilityTargetData_SingleTargetHit(Hit);
+
+		// Give it a handle and return
+		FGameplayAbilityTargetDataHandle	Handle;
+		Handle.Data = TSharedPtr<FGameplayAbilityTargetData>(TargetData);
+		OnOverlap.Broadcast(Handle);
+
+
+		// We are done. Kill us so we don't keep getting broadcast messages
+		MarkPendingKill();
+	}
+}
\ No newline at end of file
diff --git a/Engine/Source/Runtime/SkillSystem/Private/AttributeComponent.cpp b/Engine/Source/Runtime/SkillSystem/Private/AttributeComponent.cpp
--- a/Engine/Source/Runtime/SkillSystem/Private/AttributeComponent.cpp
+++ b/Engine/Source/Runtime/SkillSystem/Private/AttributeComponent.cpp
@@ -17,6 +17,17 @@ DEFINE_LOG_CATEGORY(LogAttributeComponent);
 
 #define LOCTEXT_NAMESPACE "AttributeComponent"
 
+
+int32 DebugGameplayCues = 0;
+static FAutoConsoleVariableRef CVarDebugGameplayCues(
+	TEXT("SkillSystem.DebugGameplayCues"),
+	DebugGameplayCues,
+	TEXT("Enables Debugging for GameplayCue events"),
+	ECVF_Default
+	);
+
+#pragma optimize( "", off )
+
 /** Enable to log out all render state create, destroy and updatetransform events */
 #define LOG_RENDER_STATE 0
 
@@ -167,7 +178,7 @@ FActiveGameplayEffectHandle UAttributeComponent::ApplyGameplayEffectToTarget(UGa
 {
 	check(GameplayEffect);
 
-	FGameplayEffectSpec	Spec(GameplayEffect, GetOwner(), GetOwner(), Level, GetCurveDataOverride());
+	FGameplayEffectSpec	Spec(GameplayEffect, GetOwner(),  Level, GetCurveDataOverride());
 	
 	return ApplyGameplayEffectSpecToTarget(Spec, Target, BaseQualifier);
 }
@@ -181,7 +192,9 @@ FActiveGameplayEffectHandle UAttributeComponent::K2_ApplyGameplayEffectToSelf(co
 /** This is a helper function - it seems like this will be useful as a blueprint interface at the least, but Level parameter may need to be expanded */
 FActiveGameplayEffectHandle UAttributeComponent::ApplyGameplayEffectToSelf(const UGameplayEffect *GameplayEffect, float Level, AActor *Instigator, FModifierQualifier BaseQualifier)
 {
-	FGameplayEffectSpec	Spec(GameplayEffect, GetOwner(), Instigator, Level, GetCurveDataOverride());
+	check(GameplayEffect);
+
+	FGameplayEffectSpec	Spec(GameplayEffect, Instigator, Level, GetCurveDataOverride());
 
 	return ApplyGameplayEffectSpecToSelf(Spec, BaseQualifier);
 }
@@ -239,7 +252,7 @@ FActiveGameplayEffectHandle UAttributeComponent::ApplyGameplayEffectSpecToTarget
 FActiveGameplayEffectHandle UAttributeComponent::ApplyGameplayEffectSpecToSelf(const FGameplayEffectSpec &Spec, FModifierQualifier BaseQualifier)
 {
 	// Temp, only non instance, non periodic GEs can be predictive
-	check((BaseQualifier.PredictionKey() > 0) == (Spec.GetDuration() != UGameplayEffect::INSTANT_APPLICATION && Spec.GetPeriod() == UGameplayEffect::NO_PERIOD));
+	check((BaseQualifier.PredictionKey() == 0) || (Spec.GetDuration() != UGameplayEffect::INSTANT_APPLICATION && Spec.GetPeriod() == UGameplayEffect::NO_PERIOD));
 
 	// check if the effect being applied actually succeeds
 	float ChanceToApply = Spec.GetChanceToApplyToTarget();
@@ -257,6 +270,7 @@ FActiveGameplayEffectHandle UAttributeComponent::ApplyGameplayEffectSpecToSelf(c
 	{
 	
 		float Duration = Spec.GetDuration();
+
 		if (Duration != UGameplayEffect::INSTANT_APPLICATION)
 		{
 			FActiveGameplayEffect &NewActiveEffect = ActiveGameplayEffects.CreateNewActiveGameplayEffect(Spec, BaseQualifier.PredictionKey());
@@ -381,6 +395,12 @@ float UAttributeComponent::GetGameplayEffectMagnitude(FActiveGameplayEffectHandl
 
 void UAttributeComponent::InvokeGameplayCueExecute(const FGameplayEffectSpec &Spec)
 {
+	if (DebugGameplayCues)
+	{
+		SKILL_LOG(Warning, TEXT("InvokeGameplayCueExecute: %s"), *Spec.ToSimpleString());
+	}
+
+
 	AActor *ActorOwner = GetOwner();
 	IGameplayCueInterface * GameplayCueInterface = InterfaceCast<IGameplayCueInterface>(ActorOwner);
 	if (!GameplayCueInterface)
@@ -393,7 +413,13 @@ void UAttributeComponent::InvokeGameplayCueExecute(const FGameplayEffectSpec &Sp
 	for (FGameplayEffectCue CueInfo : Spec.Def->GameplayCues)
 	{
 		float NormalizedMagnitude = CueInfo.NormalizeLevel(ExecuteLevel);
-		GameplayCueInterface->GameplayCueExecuted(CueInfo.GameplayCueTags, NormalizedMagnitude);
+		GameplayCueInterface->GameplayCueExecuted(CueInfo.GameplayCueTags, NormalizedMagnitude, Spec.InstigatorContext);
+
+		if (DebugGameplayCues)
+		{
+			DrawDebugSphere(GetWorld(), Spec.InstigatorContext.HitResult->Location, 30.f, 32, FColor(255, 0, 0), true, 30.f);
+			SKILL_LOG(Warning, TEXT("   %s"), *CueInfo.GameplayCueTags.ToString());
+		}
 	}
 }
 
@@ -411,7 +437,7 @@ void UAttributeComponent::InvokeGameplayCueActivated(const FGameplayEffectSpec &
 	for (FGameplayEffectCue CueInfo : Spec.Def->GameplayCues)
 	{
 		float NormalizedMagnitude = CueInfo.NormalizeLevel(ExecuteLevel);
-		GameplayCueInterface->GameplayCueActivated(CueInfo.GameplayCueTags, NormalizedMagnitude);
+		GameplayCueInterface->GameplayCueActivated(CueInfo.GameplayCueTags, NormalizedMagnitude, Spec.InstigatorContext);
 	}
 }
 
@@ -434,7 +460,7 @@ void UAttributeComponent::InvokeGameplayCueAdded(const FGameplayEffectSpec &Spec
 	for (FGameplayEffectCue CueInfo : Spec.Def->GameplayCues)
 	{
 		float NormalizedMagnitude = CueInfo.NormalizeLevel(ExecuteLevel);
-		GameplayCueInterface->GameplayCueAdded(CueInfo.GameplayCueTags, NormalizedMagnitude);
+		GameplayCueInterface->GameplayCueAdded(CueInfo.GameplayCueTags, NormalizedMagnitude, Spec.InstigatorContext);
 	}
 }
 
@@ -452,7 +478,7 @@ void UAttributeComponent::InvokeGameplayCueRemoved(const FGameplayEffectSpec &Sp
 	for (FGameplayEffectCue CueInfo : Spec.Def->GameplayCues)
 	{
 		float NormalizedMagnitude = CueInfo.NormalizeLevel(ExecuteLevel);
-		GameplayCueInterface->GameplayCueRemoved(CueInfo.GameplayCueTags, NormalizedMagnitude);
+		GameplayCueInterface->GameplayCueRemoved(CueInfo.GameplayCueTags, NormalizedMagnitude, Spec.InstigatorContext);
 	}
 }
 
diff --git a/Engine/Source/Runtime/SkillSystem/Private/GameplayCueView.cpp b/Engine/Source/Runtime/SkillSystem/Private/GameplayCueView.cpp
--- a/Engine/Source/Runtime/SkillSystem/Private/GameplayCueView.cpp
+++ b/Engine/Source/Runtime/SkillSystem/Private/GameplayCueView.cpp
@@ -9,6 +9,7 @@
 #include "Particles/ParticleLODLevel.h"
 #include "Particles/ParticleModuleRequired.h"
 #include "Particles/ParticleSystemComponent.h"
+#include "GameplayTagsModule.h"
 
 UGameplayCueView::UGameplayCueView(const class FPostConstructInitializeProperties& PCIP)
 : Super(PCIP)
@@ -16,44 +17,55 @@ UGameplayCueView::UGameplayCueView(const class FPostConstructInitializePropertie
 
 }
 
-
-void FGameplayCueHandler::GameplayCueActivated(const FGameplayTagContainer & GameplayCueTags, float NormalizedMagnitude)
+FGameplayCueViewInfo * FGameplayCueHandler::GetBestMatchingView(EGameplayCueEvent::Type Type, const FGameplayTag BaseTag)
 {
-	check(Owner);
-	for (auto TagIt = GameplayCueTags.CreateConstIterator(); TagIt; ++TagIt)
+	IGameplayTagsModule& GameplayTagsModule = IGameplayTagsModule::Get();
+	FGameplayTagContainer TagAndParentsContainer = GameplayTagsModule.GetGameplayTagsManager().RequestGameplayTagParents(BaseTag);
+
+	for (auto InnerTagIt = TagAndParentsContainer.CreateConstIterator(); InnerTagIt; ++InnerTagIt)
 	{
+		FGameplayTag Tag = *InnerTagIt;
+
 		for (UGameplayCueView * Def : Definitions)
 		{
 			for (FGameplayCueViewInfo & View : Def->Views)
 			{
-				if (View.CueType == EGameplayCueEvent::Applied && View.Tags.HasTag(*TagIt, EGameplayTagMatchType::IncludeParentTags, EGameplayTagMatchType::Explicit))
+				if (View.CueType == Type && View.Tags.HasTag(Tag, EGameplayTagMatchType::IncludeParentTags, EGameplayTagMatchType::Explicit))
 				{
-					View.SpawnViewEffects(Owner, NULL);
+					return &View;
 				}
 			}
 		}
 	}
+
+	return NULL;
 }
 
-void FGameplayCueHandler::GameplayCueExecuted(const FGameplayTagContainer & GameplayCueTags, float NormalizedMagnitude)
+void FGameplayCueHandler::GameplayCueActivated(const FGameplayTagContainer & GameplayCueTags, float NormalizedMagnitude, const FGameplayEffectInstigatorContext InstigatorContext)
 {
 	check(Owner);
 	for (auto TagIt = GameplayCueTags.CreateConstIterator(); TagIt; ++TagIt)
 	{
-		for (UGameplayCueView * Def : Definitions)
+		if (FGameplayCueViewInfo * View = GetBestMatchingView(EGameplayCueEvent::Applied, *TagIt))
 		{
-			for (FGameplayCueViewInfo & View : Def->Views)
-			{
-				if (View.CueType == EGameplayCueEvent::Executed && View.Tags.HasTag(*TagIt, EGameplayTagMatchType::IncludeParentTags, EGameplayTagMatchType::Explicit))
-				{
-					View.SpawnViewEffects(Owner, NULL);
-				}
-			}
+			View->SpawnViewEffects(Owner, NULL, InstigatorContext);
+		}
+	}
+}
+
+void FGameplayCueHandler::GameplayCueExecuted(const FGameplayTagContainer & GameplayCueTags, float NormalizedMagnitude, const FGameplayEffectInstigatorContext InstigatorContext)
+{
+	check(Owner);
+	for (auto TagIt = GameplayCueTags.CreateConstIterator(); TagIt; ++TagIt)
+	{
+		if (FGameplayCueViewInfo * View = GetBestMatchingView(EGameplayCueEvent::Executed, *TagIt))
+		{
+			View->SpawnViewEffects(Owner, NULL, InstigatorContext);
 		}
 	}
 }
 
-void FGameplayCueHandler::GameplayCueAdded(const FGameplayTagContainer & GameplayCueTags, float NormalizedMagnitude)
+void FGameplayCueHandler::GameplayCueAdded(const FGameplayTagContainer & GameplayCueTags, float NormalizedMagnitude, const FGameplayEffectInstigatorContext InstigatorContext)
 {
 	check(Owner);
 
@@ -65,22 +77,17 @@ void FGameplayCueHandler::GameplayCueAdded(const FGameplayTagContainer & Gamepla
 		ClearEffects(Effects);
 		check(Effects.Num() == 0);
 
-		// Add new effects
-		for (UGameplayCueView * Def : Definitions)
+		if (FGameplayCueViewInfo * View = GetBestMatchingView(EGameplayCueEvent::Added, *TagIt))
 		{
-			for (FGameplayCueViewInfo & View : Def->Views)
-			{
-				if (View.CueType == EGameplayCueEvent::Added && View.Tags.HasTag(*TagIt, EGameplayTagMatchType::IncludeParentTags, EGameplayTagMatchType::Explicit))
-				{
-					TSharedPtr<FGameplayCueViewEffects> SpawnedEffects = View.SpawnViewEffects(Owner, &SpawnedObjects);
-					Effects.Add(SpawnedEffects);
-				}
-			}
+			View->SpawnViewEffects(Owner, NULL, InstigatorContext);
+
+			TSharedPtr<FGameplayCueViewEffects> SpawnedEffects = View->SpawnViewEffects(Owner, &SpawnedObjects, InstigatorContext);
+			Effects.Add(SpawnedEffects);
 		}
 	}
 }
 
-void FGameplayCueHandler::GameplayCueRemoved(const FGameplayTagContainer & GameplayCueTags, float NormalizedMagnitude)
+void FGameplayCueHandler::GameplayCueRemoved(const FGameplayTagContainer & GameplayCueTags, float NormalizedMagnitude, const FGameplayEffectInstigatorContext InstigatorContext)
 {
 	check(Owner);
 
@@ -93,16 +100,11 @@ void FGameplayCueHandler::GameplayCueRemoved(const FGameplayTagContainer & Gamep
 			SpawnedViewEffects.Remove(*TagIt);
 		}	
 		
-		// Add new effects
-		for (UGameplayCueView * Def : Definitions)
+		// Remove old effects
+		
+		if (FGameplayCueViewInfo * View = GetBestMatchingView(EGameplayCueEvent::Removed, *TagIt))
 		{
-			for (FGameplayCueViewInfo & View : Def->Views)
-			{
-				if (View.CueType == EGameplayCueEvent::Removed && View.Tags.HasTag(*TagIt, EGameplayTagMatchType::IncludeParentTags, EGameplayTagMatchType::Explicit))
-				{
-					View.SpawnViewEffects(Owner, NULL);
-				}
-			}
+			View->SpawnViewEffects(Owner, NULL, InstigatorContext);
 		}
 	}
 }
@@ -155,7 +157,7 @@ void FGameplayCueHandler::ClearEffects(TArray< TSharedPtr<FGameplayCueViewEffect
 	}
 }
 
-TSharedPtr<FGameplayCueViewEffects> FGameplayCueViewInfo::SpawnViewEffects(AActor *Owner, TArray<UObject*> *SpawnedObjects) const
+TSharedPtr<FGameplayCueViewEffects> FGameplayCueViewInfo::SpawnViewEffects(AActor *Owner, TArray<UObject*> *SpawnedObjects, const FGameplayEffectInstigatorContext InstigatorContext) const
 {
 	check(Owner);
 
@@ -171,7 +173,18 @@ TSharedPtr<FGameplayCueViewEffects> FGameplayCueViewInfo::SpawnViewEffects(AActo
 	}
 	if (ParticleSystem)
 	{
-		SpawnedEffects->ParticleSystemComponent = UGameplayStatics::SpawnEmitterAttached(ParticleSystem, Owner->GetRootComponent());
+		if (InstigatorContext.HitResult.IsValid())
+		{
+			FHitResult &HitResult = *InstigatorContext.HitResult.Get();
+
+			SpawnedEffects->ParticleSystemComponent = UGameplayStatics::SpawnEmitterAttached(ParticleSystem, Owner->GetRootComponent(), NAME_None, HitResult.Location,
+				HitResult.Normal.Rotation(), EAttachLocation::KeepWorldPosition);
+		}
+		else
+		{
+			SpawnedEffects->ParticleSystemComponent = UGameplayStatics::SpawnEmitterAttached(ParticleSystem, Owner->GetRootComponent());
+		}		
+
 		if (SpawnedObjects)
 		{
 			SpawnedObjects->Add(SpawnedEffects->ParticleSystemComponent.Get());
@@ -185,7 +198,7 @@ TSharedPtr<FGameplayCueViewEffects> FGameplayCueViewInfo::SpawnViewEffects(AActo
 					SpawnedEffects->ParticleSystemComponent->EmitterInstances[EmitterIndx]->CurrentLODLevel->RequiredModule &&
 					SpawnedEffects->ParticleSystemComponent->EmitterInstances[EmitterIndx]->CurrentLODLevel->RequiredModule->EmitterLoops == 0)
 				{
-					SKILL_LOG(Warning, TEXT("%s - particle system has a looping emitter...."), *SpawnedEffects->ParticleSystemComponent->GetName());
+					SKILL_LOG(Warning, TEXT("%s - particle system has a looping emitter. This should not be used in a executed GameplayCue!"), *SpawnedEffects->ParticleSystemComponent->GetName());
 					break;
 				}
 			}
@@ -203,4 +216,4 @@ TSharedPtr<FGameplayCueViewEffects> FGameplayCueViewInfo::SpawnViewEffects(AActo
 	}
 
 	return SpawnedEffects;
-}
\ No newline at end of file
+}
diff --git a/Engine/Source/Runtime/SkillSystem/Private/GameplayEffect.cpp b/Engine/Source/Runtime/SkillSystem/Private/GameplayEffect.cpp
--- a/Engine/Source/Runtime/SkillSystem/Private/GameplayEffect.cpp
+++ b/Engine/Source/Runtime/SkillSystem/Private/GameplayEffect.cpp
@@ -11,6 +11,8 @@ const float UGameplayEffect::INSTANT_APPLICATION = 0.f;
 const float UGameplayEffect::NO_PERIOD = 0.f;
 const float FGameplayEffectLevelSpec::INVALID_LEVEL = -1.f;
 
+#pragma optimize( "", off )
+
 #if SKILL_SYSTEM_AGGREGATOR_DEBUG
 FAggregator::FAllocationStats FAggregator::AllocationStats;
 #endif
@@ -50,9 +52,9 @@ bool UGameplayEffect::AreApplicationTagRequirementsSatisfied(const FGameplayTagC
 //
 // --------------------------------------------------------------------------------------------------------------------------------------------------------
 
-FGameplayEffectSpec::FGameplayEffectSpec(const UGameplayEffect *InDef, AActor *Owner, AActor *Instigator, float Level, const FGlobalCurveDataOverride *CurveData)
+FGameplayEffectSpec::FGameplayEffectSpec(const UGameplayEffect *InDef, AActor * Instigator, float Level, const FGlobalCurveDataOverride *CurveData)
 : Def(InDef)
-, ModifierLevel(TSharedPtr<FGameplayEffectLevelSpec>(new FGameplayEffectLevelSpec(Level, InDef->LevelInfo, Owner)))
+, ModifierLevel(TSharedPtr<FGameplayEffectLevelSpec>(new FGameplayEffectLevelSpec(Level, InDef->LevelInfo, Instigator)))
 , Duration(new FAggregator(InDef->Duration.MakeFinalizedCopy(CurveData), ModifierLevel, SKILL_AGG_DEBUG(TEXT("%s Duration"), *InDef->GetName())))
 , Period(new FAggregator(InDef->Period.MakeFinalizedCopy(CurveData), ModifierLevel, SKILL_AGG_DEBUG(TEXT("%s Period"), *InDef->GetName())))
 , ChanceToApplyToTarget(new FAggregator(InDef->ChanceToApplyToTarget.MakeFinalizedCopy(CurveData), ModifierLevel, SKILL_AGG_DEBUG(TEXT("%s ChanceToApplyToTarget"), *InDef->GetName())))
@@ -66,6 +68,7 @@ FGameplayEffectSpec::FGameplayEffectSpec(const UGameplayEffect *InDef, AActor *O
 	ChanceToExecuteOnGameplayEffect.Get()->RegisterLevelDependancies();
 
 	// check the stacking policy for this gameplay effect and warn if it's set up badly
+	// FIXME: this check should be done in the UGameplayEffect PostEdit or PostLoad - not in the spec startup
 	if (Def->StackingPolicy != EGameplayEffectStackingPolicy::Unlimited)
 	{
 		StackingPolicy = Def->StackingPolicy;
@@ -143,17 +146,17 @@ FGameplayEffectSpec::FGameplayEffectSpec(const UGameplayEffect *InDef, AActor *O
 		}
 	}
 
-	InitModifiers(CurveData, Owner, Level);
+	InitModifiers(CurveData, Instigator, Level);
 
 	if (InDef)
 	{
 		for (UGameplayEffect* TargetDef : InDef->TargetEffects)
 		{
-			TargetEffectSpecs.Add(TSharedRef<FGameplayEffectSpec>(new FGameplayEffectSpec(TargetDef, Owner, Instigator, Level, CurveData)));
+			TargetEffectSpecs.Add(TSharedRef<FGameplayEffectSpec>(new FGameplayEffectSpec(TargetDef, Instigator, Level, CurveData)));
 		}
 	}
 
-	InstigatorStack.AddInstigator(Instigator);
+	InstigatorContext.AddInstigator(Instigator);
 }
 
 void FGameplayEffectSpec::InitModifiers(const FGlobalCurveDataOverride *CurveData, AActor *Owner, float Level)
@@ -412,7 +415,7 @@ FModifierSpec::FModifierSpec(const FGameplayModifierInfo &InInfo, TSharedPtr<FGa
 	Aggregator.Get()->RegisterLevelDependancies();
 	if (InInfo.TargetEffect)
 	{
-		TargetEffectSpec = TSharedPtr<FGameplayEffectSpec>(new FGameplayEffectSpec(InInfo.TargetEffect, Owner, Owner, Level, CurveData));
+		TargetEffectSpec = TSharedPtr<FGameplayEffectSpec>(new FGameplayEffectSpec(InInfo.TargetEffect, Owner, Level, CurveData));
 	}
 }
 
@@ -446,6 +449,7 @@ bool FModifierSpec::AreTagRequirementsSatisfied(const FModifierSpec &ModifierToB
 void FGameplayEffectInstigatorContext::AddInstigator(class AActor *InInstigator)
 {
 	Instigator = InInstigator;
+	InstigatorAttributeComponent = NULL;
 
 	// Cache off his attribute component.
 	if (Instigator)
@@ -464,6 +468,41 @@ void FGameplayEffectInstigatorContext::AddInstigator(class AActor *InInstigator)
 	}
 }
 
+void FGameplayEffectInstigatorContext::AddHitResult(const FHitResult InHitResult)
+{
+	check(!HitResult.IsValid());
+	HitResult = TSharedPtr<FHitResult>(new FHitResult(InHitResult));
+}
+
+bool FGameplayEffectInstigatorContext::NetSerialize(FArchive& Ar, class UPackageMap* Map, bool& bOutSuccess)
+{
+	Ar << Instigator;
+
+	bool HasHitResults = HitResult.IsValid();
+	Ar << HasHitResults;
+	if (Ar.IsLoading())
+	{
+		if (HasHitResults)
+		{
+			if (!HitResult.IsValid())
+			{
+				HitResult = TSharedPtr<FHitResult>(new FHitResult());
+			}
+			AddInstigator(Instigator); // Just to initialize InstigatorAttributeComponent
+		}
+	}
+
+	if (HasHitResults == 1)
+	{
+		Ar << HitResult->BoneName;
+		HitResult->Location.NetSerialize(Ar, Map, bOutSuccess);
+		HitResult->Normal.NetSerialize(Ar, Map, bOutSuccess);
+	}
+
+	bOutSuccess = true;
+	return true;
+}
+
 // --------------------------------------------------------------------------------------------------------------------------------------------------------
 //
 //	FAggregatorRef
@@ -1190,7 +1229,7 @@ void FActiveGameplayEffectsContainer::ExecuteActiveEffectsFrom(const FGameplayEf
 	// check if the new effect removes or modifies existing effects
 	if (Spec.Def->AppliesImmunityTo == EGameplayImmunity::ActiveGE)
 	{
-		for (int ii = 0; ii < GameplayEffects.Num(); ++ii)
+		for (int32 ii = 0; ii < GameplayEffects.Num(); ++ii)
 		{
 			// should we try to apply this to the current effect?
 			if ((ChanceToExecute < 1.f - SMALL_NUMBER) && (FMath::FRand() > ChanceToExecute))
@@ -1206,7 +1245,7 @@ void FActiveGameplayEffectsContainer::ExecuteActiveEffectsFrom(const FGameplayEf
 			}
 
 			// we kept the effect, now check its mods to see if we should remove any of them
-			for (int jj = 0; jj < GameplayEffects[ii].Spec.Modifiers.Num(); ++jj)
+			for (int32 jj = 0; jj < GameplayEffects[ii].Spec.Modifiers.Num(); ++jj)
 			{
 				const FGameplayModifierEvaluatedData& Data = GameplayEffects[ii].Spec.Modifiers[jj].Aggregator.Get()->Evaluate();
 				if (GameplayEffects[ii].Spec.Def->AreGameplayEffectTagRequirementsSatisfied(Data.Tags))
@@ -1452,13 +1491,16 @@ FActiveGameplayEffect & FActiveGameplayEffectsContainer::CreateNewActiveGameplay
 {
 	LastAssignedHandle = LastAssignedHandle.GetNextHandle();
 	FActiveGameplayEffect & NewEffect = *new (GameplayEffects)FActiveGameplayEffect(LastAssignedHandle, Spec, GetWorldTime(), GetGameStateTime(), InPredictionKey);
-	
-	if (InPredictionKey == 0 || IsNetAuthority())	// Clients predicting a GameplayEffect must not call MarkItemDirty
+
+	if (InPredictionKey == 0 || IsNetAuthority())	// Clients predicting a GameplayEffect should not call MarkItemDirty
 	{
 		MarkItemDirty(NewEffect);
 	}
 	else
 	{
+		// Clients predicting should call MarkArrayDirty to force the internal replication map to be rebuilt.
+		MarkArrayDirty();
+
 		Owner->GetPredictionKeyDelegate(InPredictionKey).AddUObject(Owner, &UAttributeComponent::RemoveActiveGameplayEffect_NoReturn, NewEffect.Handle);
 	}
 	return NewEffect;
@@ -1560,11 +1602,15 @@ void FActiveGameplayEffectsContainer::TEMP_TickActiveEffects(float DeltaSeconds)
 		if (Duration > 0.f && Effect.StartWorldTime + Effect.GetDuration() < CurrentTime)
 		{
 			Owner->InvokeGameplayCueRemoved(Effect.Spec);
-
 			MarkArrayDirty();
 			GameplayEffects.RemoveAtSwap(Idx);
 			Idx--;
 
+			// Hack: force netupdate on owner. This isn't really necessary in real gameplay but is nice
+			// during debugging where breakpoints or pausing can messup network update times. Open issue
+			// with network team.
+			Owner->GetOwner()->ForceNetUpdate();
+
 			// todo: check if the removed effect deals with an attribute that stacks, if not don't set the flag
 			bNeedToRecalculateStacks = true;
 		}
@@ -1736,7 +1782,7 @@ bool FActiveGameplayEffectsContainer::CanApplyAttributeModifiers(const UGameplay
 {
 	SCOPE_CYCLE_COUNTER(STAT_GameplayEffectsCanApplyAttributeModifiers);
 
-	FGameplayEffectSpec	Spec(GameplayEffect, Instigator, Instigator, Level, Owner->GetCurveDataOverride());
+	FGameplayEffectSpec	Spec(GameplayEffect, Instigator, Level, Owner->GetCurveDataOverride());
 
 	ApplyActiveEffectsTo(Spec, FModifierQualifier().Type(EGameplayMod::IncomingGE));
 	
diff --git a/Engine/Source/Runtime/SkillSystem/Private/GameplayEffectExtension_LifestealTest.cpp b/Engine/Source/Runtime/SkillSystem/Private/GameplayEffectExtension_LifestealTest.cpp
--- a/Engine/Source/Runtime/SkillSystem/Private/GameplayEffectExtension_LifestealTest.cpp
+++ b/Engine/Source/Runtime/SkillSystem/Private/GameplayEffectExtension_LifestealTest.cpp
@@ -28,7 +28,7 @@ void UGameplayEffectExtension_LifestealTest::PostGameplayEffectExecute(const FGa
 	float HealthToRestore = -DamageDone * LifestealPCT;
 	if (HealthToRestore > 0.f)
 	{
-		UAttributeComponent *Source = Data.EffectSpec.InstigatorStack.GetOriginInstigatorAttributeComponent();
+		UAttributeComponent *Source = Data.EffectSpec.InstigatorContext.GetOriginalInstigatorAttributeComponent();
 
 		UGameplayEffect * LocalHealthRestore = HealthRestoreGameplayEffect;
 		if (!LocalHealthRestore)
diff --git a/Engine/Source/Runtime/SkillSystem/Private/GameplayEffectExtension_ShieldTest.cpp b/Engine/Source/Runtime/SkillSystem/Private/GameplayEffectExtension_ShieldTest.cpp
--- a/Engine/Source/Runtime/SkillSystem/Private/GameplayEffectExtension_ShieldTest.cpp
+++ b/Engine/Source/Runtime/SkillSystem/Private/GameplayEffectExtension_ShieldTest.cpp
@@ -16,7 +16,7 @@ UGameplayEffectExtension_ShieldTest::UGameplayEffectExtension_ShieldTest(const c
 void UGameplayEffectExtension_ShieldTest::PreGameplayEffectExecute(const FGameplayModifierEvaluatedData &SelfData, FGameplayEffectModCallbackData &Data) const
 {
 	IGameplayTagsModule& GameplayTagsModule = IGameplayTagsModule::Get();
-	UAttributeComponent *Source = Data.EffectSpec.InstigatorStack.GetOriginInstigatorAttributeComponent();
+	UAttributeComponent *Source = Data.EffectSpec.InstigatorContext.GetOriginalInstigatorAttributeComponent();
 
 	// FIXME: some annoyances here: Damage about to be applied = Data.EvaluatedData.Magnitude = negative. Do some sign flipping here that would make more sense if we were dealing with
 	// 'damage' (positive) instead of 'health' (negative)
diff --git a/Engine/Source/Runtime/SkillSystem/Private/GameplayEffectTests.cpp b/Engine/Source/Runtime/SkillSystem/Private/GameplayEffectTests.cpp
--- a/Engine/Source/Runtime/SkillSystem/Private/GameplayEffectTests.cpp
+++ b/Engine/Source/Runtime/SkillSystem/Private/GameplayEffectTests.cpp
@@ -614,8 +614,9 @@ bool GameplayEffectsTest_SnapshotBuffDestroy(UWorld *World, FAutomationTestBase
 
 	UProperty *HealthProperty = FindFieldChecked<UProperty>(USkillSystemTestAttributeSet::StaticClass(), GET_MEMBER_NAME_CHECKED(USkillSystemTestAttributeSet, Health));
 
-	UAttributeComponent * SourceComponent = Cast<UAttributeComponent>(SourceActor->CreateComponentFromTemplate(UAttributeComponent::StaticClass()->GetDefaultObject<UAttributeComponent>(), FString(TEXT("AttributeComponent"))));
-	UAttributeComponent * DestComponent = Cast<UAttributeComponent>(DestActor->CreateComponentFromTemplate(UAttributeComponent::StaticClass()->GetDefaultObject<UAttributeComponent>(), FString(TEXT("AttributeComponent"))));
+	UAttributeComponent * SourceComponent = SourceActor->AttributeComponent;
+	UAttributeComponent * DestComponent = DestActor->AttributeComponent;
+	
 	SourceComponent->GetSet<USkillSystemTestAttributeSet>()->Health = StartHealth;
 	DestComponent->GetSet<USkillSystemTestAttributeSet>()->Health = StartHealth;
 
@@ -704,8 +705,9 @@ bool GameplayEffectsTest_DurationBuff(UWorld *World, FAutomationTestBase * Test)
 
 	UProperty *HealthProperty = FindFieldChecked<UProperty>(USkillSystemTestAttributeSet::StaticClass(), GET_MEMBER_NAME_CHECKED(USkillSystemTestAttributeSet, Health));
 
-	UAttributeComponent * SourceComponent = Cast<UAttributeComponent>(SourceActor->CreateComponentFromTemplate(UAttributeComponent::StaticClass()->GetDefaultObject<UAttributeComponent>(), FString(TEXT("AttributeComponent"))));
-	UAttributeComponent * DestComponent = Cast<UAttributeComponent>(DestActor->CreateComponentFromTemplate(UAttributeComponent::StaticClass()->GetDefaultObject<UAttributeComponent>(), FString(TEXT("AttributeComponent"))));
+	UAttributeComponent * SourceComponent = SourceActor->AttributeComponent;
+	UAttributeComponent * DestComponent = DestActor->AttributeComponent;
+
 	SourceComponent->GetSet<USkillSystemTestAttributeSet>()->Health = StartHealth;
 	DestComponent->GetSet<USkillSystemTestAttributeSet>()->Health = StartHealth;
 
@@ -794,8 +796,8 @@ bool GameplayEffectsTest_DamageBuffBuff_Basic(UWorld *World, FAutomationTestBase
 
 	UProperty *HealthProperty = FindFieldChecked<UProperty>(USkillSystemTestAttributeSet::StaticClass(), GET_MEMBER_NAME_CHECKED(USkillSystemTestAttributeSet, Health));
 
-	UAttributeComponent * SourceComponent = Cast<UAttributeComponent>(SourceActor->CreateComponentFromTemplate(UAttributeComponent::StaticClass()->GetDefaultObject<UAttributeComponent>(), FString(TEXT("AttributeComponent"))));
-	UAttributeComponent * DestComponent = Cast<UAttributeComponent>(DestActor->CreateComponentFromTemplate(UAttributeComponent::StaticClass()->GetDefaultObject<UAttributeComponent>(), FString(TEXT("AttributeComponent"))));
+	UAttributeComponent * SourceComponent = SourceActor->AttributeComponent;
+	UAttributeComponent * DestComponent = DestActor->AttributeComponent;
 	SourceComponent->GetSet<USkillSystemTestAttributeSet>()->Health = StartHealth;
 	DestComponent->GetSet<USkillSystemTestAttributeSet>()->Health = StartHealth;
 
@@ -935,8 +937,8 @@ bool GameplayEffectsTest_DamageBuffBuff_FullLink(UWorld *World, FAutomationTestB
 
 	UProperty *HealthProperty = FindFieldChecked<UProperty>(USkillSystemTestAttributeSet::StaticClass(), GET_MEMBER_NAME_CHECKED(USkillSystemTestAttributeSet, Health));
 
-	UAttributeComponent * SourceComponent = Cast<UAttributeComponent>(SourceActor->CreateComponentFromTemplate(UAttributeComponent::StaticClass()->GetDefaultObject<UAttributeComponent>(), FString(TEXT("AttributeComponent"))));
-	UAttributeComponent * DestComponent = Cast<UAttributeComponent>(DestActor->CreateComponentFromTemplate(UAttributeComponent::StaticClass()->GetDefaultObject<UAttributeComponent>(), FString(TEXT("AttributeComponent"))));
+	UAttributeComponent * SourceComponent = SourceActor->AttributeComponent;
+	UAttributeComponent * DestComponent = DestActor->AttributeComponent;
 	SourceComponent->GetSet<USkillSystemTestAttributeSet>()->Health = StartHealth;
 	DestComponent->GetSet<USkillSystemTestAttributeSet>()->Health = StartHealth;
 
@@ -1073,8 +1075,8 @@ bool GameplayEffectsTest_DamageBuffBuff_FullSnapshot(UWorld *World, FAutomationT
 
 	UProperty *HealthProperty = FindFieldChecked<UProperty>(USkillSystemTestAttributeSet::StaticClass(), GET_MEMBER_NAME_CHECKED(USkillSystemTestAttributeSet, Health));
 
-	UAttributeComponent * SourceComponent = Cast<UAttributeComponent>(SourceActor->CreateComponentFromTemplate(UAttributeComponent::StaticClass()->GetDefaultObject<UAttributeComponent>(), FString(TEXT("AttributeComponent"))));
-	UAttributeComponent * DestComponent = Cast<UAttributeComponent>(DestActor->CreateComponentFromTemplate(UAttributeComponent::StaticClass()->GetDefaultObject<UAttributeComponent>(), FString(TEXT("AttributeComponent"))));
+	UAttributeComponent * SourceComponent = SourceActor->AttributeComponent;
+	UAttributeComponent * DestComponent = DestActor->AttributeComponent;
 	SourceComponent->GetSet<USkillSystemTestAttributeSet>()->Health = StartHealth;
 	DestComponent->GetSet<USkillSystemTestAttributeSet>()->Health = StartHealth;
 
@@ -1211,8 +1213,8 @@ bool GameplayEffectsTest_DamageBuffBuff_SnapshotLink(UWorld *World, FAutomationT
 
 	UProperty *HealthProperty = FindFieldChecked<UProperty>(USkillSystemTestAttributeSet::StaticClass(), GET_MEMBER_NAME_CHECKED(USkillSystemTestAttributeSet, Health));
 
-	UAttributeComponent * SourceComponent = Cast<UAttributeComponent>(SourceActor->CreateComponentFromTemplate(UAttributeComponent::StaticClass()->GetDefaultObject<UAttributeComponent>(), FString(TEXT("AttributeComponent"))));
-	UAttributeComponent * DestComponent = Cast<UAttributeComponent>(DestActor->CreateComponentFromTemplate(UAttributeComponent::StaticClass()->GetDefaultObject<UAttributeComponent>(), FString(TEXT("AttributeComponent"))));
+	UAttributeComponent * SourceComponent = SourceActor->AttributeComponent;
+	UAttributeComponent * DestComponent = DestActor->AttributeComponent;
 	SourceComponent->GetSet<USkillSystemTestAttributeSet>()->Health = StartHealth;
 	DestComponent->GetSet<USkillSystemTestAttributeSet>()->Health = StartHealth;
 
@@ -1345,10 +1347,8 @@ bool TimerTest(UWorld *World, FAutomationTestBase * Test)
 
 	UProperty *HealthProperty = FindFieldChecked<UProperty>(USkillSystemTestAttributeSet::StaticClass(), GET_MEMBER_NAME_CHECKED(USkillSystemTestAttributeSet, Health));
 
-	UAttributeComponent * SourceComponent = Cast<UAttributeComponent>(SourceActor->CreateComponentFromTemplate(UAttributeComponent::StaticClass()->GetDefaultObject<UAttributeComponent>(), FString(TEXT("AttributeComponent"))));
-	UAttributeComponent * DestComponent = Cast<UAttributeComponent>(DestActor->CreateComponentFromTemplate(UAttributeComponent::StaticClass()->GetDefaultObject<UAttributeComponent>(), FString(TEXT("AttributeComponent"))));
-	SourceComponent->RegisterComponentWithWorld(World);
-	DestComponent->RegisterComponentWithWorld(World);
+	UAttributeComponent * SourceComponent = SourceActor->AttributeComponent;
+	UAttributeComponent * DestComponent = DestActor->AttributeComponent;
 
 	SourceComponent->GetSet<USkillSystemTestAttributeSet>()->Health = StartHealth;
 	DestComponent->GetSet<USkillSystemTestAttributeSet>()->Health = StartHealth;
@@ -1703,10 +1703,8 @@ bool GameplayEffectsTest_DurationDamage(UWorld *World, FAutomationTestBase * Tes
 
 	UProperty *HealthProperty = FindFieldChecked<UProperty>(USkillSystemTestAttributeSet::StaticClass(), GET_MEMBER_NAME_CHECKED(USkillSystemTestAttributeSet, Health));
 
-	UAttributeComponent * SourceComponent = Cast<UAttributeComponent>(SourceActor->CreateComponentFromTemplate(UAttributeComponent::StaticClass()->GetDefaultObject<UAttributeComponent>(), FString(TEXT("AttributeComponent"))));
-	UAttributeComponent * DestComponent = Cast<UAttributeComponent>(DestActor->CreateComponentFromTemplate(UAttributeComponent::StaticClass()->GetDefaultObject<UAttributeComponent>(), FString(TEXT("AttributeComponent"))));
-	SourceComponent->RegisterComponentWithWorld(World);
-	DestComponent->RegisterComponentWithWorld(World);
+	UAttributeComponent * SourceComponent = SourceActor->AttributeComponent;
+	UAttributeComponent * DestComponent = DestActor->AttributeComponent;
 
 	SourceComponent->GetSet<USkillSystemTestAttributeSet>()->Health = StartHealth;
 	DestComponent->GetSet<USkillSystemTestAttributeSet>()->Health = StartHealth;
@@ -1779,10 +1777,8 @@ bool GameplayEffectsTest_PeriodicDamage(UWorld *World, FAutomationTestBase * Tes
 
 	UProperty *HealthProperty = FindFieldChecked<UProperty>(USkillSystemTestAttributeSet::StaticClass(), GET_MEMBER_NAME_CHECKED(USkillSystemTestAttributeSet, Health));
 
-	UAttributeComponent * SourceComponent = Cast<UAttributeComponent>(SourceActor->CreateComponentFromTemplate(UAttributeComponent::StaticClass()->GetDefaultObject<UAttributeComponent>(), FString(TEXT("AttributeComponent"))));
-	UAttributeComponent * DestComponent = Cast<UAttributeComponent>(DestActor->CreateComponentFromTemplate(UAttributeComponent::StaticClass()->GetDefaultObject<UAttributeComponent>(), FString(TEXT("AttributeComponent"))));
-	SourceComponent->RegisterComponentWithWorld(World);
-	DestComponent->RegisterComponentWithWorld(World);
+	UAttributeComponent * SourceComponent = SourceActor->AttributeComponent;
+	UAttributeComponent * DestComponent = DestActor->AttributeComponent;
 
 	SourceComponent->GetSet<USkillSystemTestAttributeSet>()->Health = StartHealth;
 	DestComponent->GetSet<USkillSystemTestAttributeSet>()->Health = StartHealth;
@@ -1853,10 +1849,8 @@ bool GameplayEffectsTest_LifestealExtension(UWorld *World, FAutomationTestBase *
 	UProperty *HealthProperty = FindFieldChecked<UProperty>(USkillSystemTestAttributeSet::StaticClass(), GET_MEMBER_NAME_CHECKED(USkillSystemTestAttributeSet, Health));
 	UProperty *DamageProperty = FindFieldChecked<UProperty>(USkillSystemTestAttributeSet::StaticClass(), GET_MEMBER_NAME_CHECKED(USkillSystemTestAttributeSet, Damage));
 
-	UAttributeComponent * SourceComponent = Cast<UAttributeComponent>(SourceActor->CreateComponentFromTemplate(UAttributeComponent::StaticClass()->GetDefaultObject<UAttributeComponent>(), FString(TEXT("AttributeComponent"))));
-	UAttributeComponent * DestComponent = Cast<UAttributeComponent>(DestActor->CreateComponentFromTemplate(UAttributeComponent::StaticClass()->GetDefaultObject<UAttributeComponent>(), FString(TEXT("AttributeComponent"))));
-	SourceComponent->RegisterComponentWithWorld(World);
-	DestComponent->RegisterComponentWithWorld(World);
+	UAttributeComponent * SourceComponent = SourceActor->AttributeComponent;
+	UAttributeComponent * DestComponent = DestActor->AttributeComponent;
 
 	SourceComponent->GetSet<USkillSystemTestAttributeSet>()->Health = StartHealth;
 	DestComponent->GetSet<USkillSystemTestAttributeSet>()->Health = StartHealth;
@@ -1930,10 +1924,8 @@ bool GameplayEffectsTest_ShieldExtension(UWorld *World, FAutomationTestBase * Te
 	UProperty *HealthProperty = FindFieldChecked<UProperty>(USkillSystemTestAttributeSet::StaticClass(), GET_MEMBER_NAME_CHECKED(USkillSystemTestAttributeSet, Health));
 	UProperty *DamageProperty = FindFieldChecked<UProperty>(USkillSystemTestAttributeSet::StaticClass(), GET_MEMBER_NAME_CHECKED(USkillSystemTestAttributeSet, Damage));
 
-	UAttributeComponent * SourceComponent = Cast<UAttributeComponent>(SourceActor->CreateComponentFromTemplate(UAttributeComponent::StaticClass()->GetDefaultObject<UAttributeComponent>(), FString(TEXT("AttributeComponent"))));
-	UAttributeComponent * DestComponent = Cast<UAttributeComponent>(DestActor->CreateComponentFromTemplate(UAttributeComponent::StaticClass()->GetDefaultObject<UAttributeComponent>(), FString(TEXT("AttributeComponent"))));
-	SourceComponent->RegisterComponentWithWorld(World);
-	DestComponent->RegisterComponentWithWorld(World);
+	UAttributeComponent * SourceComponent = SourceActor->AttributeComponent;
+	UAttributeComponent * DestComponent = DestActor->AttributeComponent;
 
 	SourceComponent->GetSet<USkillSystemTestAttributeSet>()->Health = StartHealth;
 	DestComponent->GetSet<USkillSystemTestAttributeSet>()->Health = StartHealth;
@@ -2047,10 +2039,8 @@ bool GameplayEffectsTest_ShieldExtensionMultiple(UWorld *World, FAutomationTestB
 	UProperty *HealthProperty = FindFieldChecked<UProperty>(USkillSystemTestAttributeSet::StaticClass(), GET_MEMBER_NAME_CHECKED(USkillSystemTestAttributeSet, Health));
 	UProperty *DamageProperty = FindFieldChecked<UProperty>(USkillSystemTestAttributeSet::StaticClass(), GET_MEMBER_NAME_CHECKED(USkillSystemTestAttributeSet, Damage));
 
-	UAttributeComponent * SourceComponent = Cast<UAttributeComponent>(SourceActor->CreateComponentFromTemplate(UAttributeComponent::StaticClass()->GetDefaultObject<UAttributeComponent>(), FString(TEXT("AttributeComponent"))));
-	UAttributeComponent * DestComponent = Cast<UAttributeComponent>(DestActor->CreateComponentFromTemplate(UAttributeComponent::StaticClass()->GetDefaultObject<UAttributeComponent>(), FString(TEXT("AttributeComponent"))));
-	SourceComponent->RegisterComponentWithWorld(World);
-	DestComponent->RegisterComponentWithWorld(World);
+	UAttributeComponent * SourceComponent = SourceActor->AttributeComponent;
+	UAttributeComponent * DestComponent = DestActor->AttributeComponent;
 
 	SourceComponent->GetSet<USkillSystemTestAttributeSet>()->Health = StartHealth;
 	DestComponent->GetSet<USkillSystemTestAttributeSet>()->Health = StartHealth;
@@ -4862,8 +4852,9 @@ bool GameplayEffectsTest_ChanceToApplyToTarget(UWorld *World, FAutomationTestBas
 
 	UProperty *DamageProperty = FindFieldChecked<UProperty>(USkillSystemTestAttributeSet::StaticClass(), GET_MEMBER_NAME_CHECKED(USkillSystemTestAttributeSet, Damage));
 
-	UAttributeComponent * SourceComponent = Cast<UAttributeComponent>(SourceActor->CreateComponentFromTemplate(UAttributeComponent::StaticClass()->GetDefaultObject<UAttributeComponent>(), FString(TEXT("AttributeComponent"))));
-	UAttributeComponent * DestComponent = Cast<UAttributeComponent>(DestActor->CreateComponentFromTemplate(UAttributeComponent::StaticClass()->GetDefaultObject<UAttributeComponent>(), FString(TEXT("AttributeComponent"))));
+	UAttributeComponent * SourceComponent = SourceActor->AttributeComponent;
+	UAttributeComponent * DestComponent = DestActor->AttributeComponent;
+
 	SourceComponent->GetSet<USkillSystemTestAttributeSet>()->Health = StartHealth;
 	DestComponent->GetSet<USkillSystemTestAttributeSet>()->Health = StartHealth;
 
@@ -4930,8 +4921,9 @@ bool GameplayEffectsTest_ChanceToExecuteOnActiveGEMod(UWorld *World, FAutomation
 
 	UProperty *DamageProperty = FindFieldChecked<UProperty>(USkillSystemTestAttributeSet::StaticClass(), GET_MEMBER_NAME_CHECKED(USkillSystemTestAttributeSet, Damage));
 
-	UAttributeComponent * SourceComponent = Cast<UAttributeComponent>(SourceActor->CreateComponentFromTemplate(UAttributeComponent::StaticClass()->GetDefaultObject<UAttributeComponent>(), FString(TEXT("AttributeComponent"))));
-	UAttributeComponent * DestComponent = Cast<UAttributeComponent>(DestActor->CreateComponentFromTemplate(UAttributeComponent::StaticClass()->GetDefaultObject<UAttributeComponent>(), FString(TEXT("AttributeComponent"))));
+	UAttributeComponent * SourceComponent = SourceActor->AttributeComponent;
+	UAttributeComponent * DestComponent = DestActor->AttributeComponent;
+
 	SourceComponent->GetSet<USkillSystemTestAttributeSet>()->Health = StartHealth;
 	DestComponent->GetSet<USkillSystemTestAttributeSet>()->Health = StartHealth;
 
@@ -5046,8 +5038,9 @@ bool GameplayEffectsTest_ChanceToExecuteOnActiveGEImmunity(UWorld *World, FAutom
 
 	UProperty *DamageProperty = FindFieldChecked<UProperty>(USkillSystemTestAttributeSet::StaticClass(), GET_MEMBER_NAME_CHECKED(USkillSystemTestAttributeSet, Damage));
 
-	UAttributeComponent * SourceComponent = Cast<UAttributeComponent>(SourceActor->CreateComponentFromTemplate(UAttributeComponent::StaticClass()->GetDefaultObject<UAttributeComponent>(), FString(TEXT("AttributeComponent"))));
-	UAttributeComponent * DestComponent = Cast<UAttributeComponent>(DestActor->CreateComponentFromTemplate(UAttributeComponent::StaticClass()->GetDefaultObject<UAttributeComponent>(), FString(TEXT("AttributeComponent"))));
+	UAttributeComponent * SourceComponent = SourceActor->AttributeComponent;
+	UAttributeComponent * DestComponent = DestActor->AttributeComponent;
+
 	SourceComponent->GetSet<USkillSystemTestAttributeSet>()->Health = StartHealth;
 	DestComponent->GetSet<USkillSystemTestAttributeSet>()->Health = StartHealth;
 
@@ -5155,9 +5148,9 @@ bool GameplayEffectsTest_ChanceToExecuteOnOutgoingGEMod(UWorld *World, FAutomati
 	ASkillSystemTestPawn *DestActor = World->SpawnActor<ASkillSystemTestPawn>();
 
 	UProperty *DamageProperty = FindFieldChecked<UProperty>(USkillSystemTestAttributeSet::StaticClass(), GET_MEMBER_NAME_CHECKED(USkillSystemTestAttributeSet, Damage));
-
-	UAttributeComponent * SourceComponent = Cast<UAttributeComponent>(SourceActor->CreateComponentFromTemplate(UAttributeComponent::StaticClass()->GetDefaultObject<UAttributeComponent>(), FString(TEXT("AttributeComponent"))));
-	UAttributeComponent * DestComponent = Cast<UAttributeComponent>(DestActor->CreateComponentFromTemplate(UAttributeComponent::StaticClass()->GetDefaultObject<UAttributeComponent>(), FString(TEXT("AttributeComponent"))));
+	
+	UAttributeComponent * SourceComponent = SourceActor->AttributeComponent;
+	UAttributeComponent * DestComponent = DestActor->AttributeComponent;
 	SourceComponent->GetSet<USkillSystemTestAttributeSet>()->Health = StartHealth;
 	DestComponent->GetSet<USkillSystemTestAttributeSet>()->Health = StartHealth;
 
@@ -5259,8 +5252,8 @@ bool GameplayEffectsTest_ChanceToExecuteOnOutgoingGEImmunity(UWorld *World, FAut
 
 	UProperty *DamageProperty = FindFieldChecked<UProperty>(USkillSystemTestAttributeSet::StaticClass(), GET_MEMBER_NAME_CHECKED(USkillSystemTestAttributeSet, Damage));
 
-	UAttributeComponent * SourceComponent = Cast<UAttributeComponent>(SourceActor->CreateComponentFromTemplate(UAttributeComponent::StaticClass()->GetDefaultObject<UAttributeComponent>(), FString(TEXT("AttributeComponent"))));
-	UAttributeComponent * DestComponent = Cast<UAttributeComponent>(DestActor->CreateComponentFromTemplate(UAttributeComponent::StaticClass()->GetDefaultObject<UAttributeComponent>(), FString(TEXT("AttributeComponent"))));
+	UAttributeComponent * SourceComponent = SourceActor->AttributeComponent;
+	UAttributeComponent * DestComponent = DestActor->AttributeComponent;
 	SourceComponent->GetSet<USkillSystemTestAttributeSet>()->Health = StartHealth;
 	DestComponent->GetSet<USkillSystemTestAttributeSet>()->Health = StartHealth;
 
@@ -5352,8 +5345,8 @@ bool GameplayEffectsTest_ChanceToExecuteOnIncomingGEMod(UWorld *World, FAutomati
 
 	UProperty *DamageProperty = FindFieldChecked<UProperty>(USkillSystemTestAttributeSet::StaticClass(), GET_MEMBER_NAME_CHECKED(USkillSystemTestAttributeSet, Damage));
 
-	UAttributeComponent * SourceComponent = Cast<UAttributeComponent>(SourceActor->CreateComponentFromTemplate(UAttributeComponent::StaticClass()->GetDefaultObject<UAttributeComponent>(), FString(TEXT("AttributeComponent"))));
-	UAttributeComponent * DestComponent = Cast<UAttributeComponent>(DestActor->CreateComponentFromTemplate(UAttributeComponent::StaticClass()->GetDefaultObject<UAttributeComponent>(), FString(TEXT("AttributeComponent"))));
+	UAttributeComponent * SourceComponent = SourceActor->AttributeComponent;
+	UAttributeComponent * DestComponent = DestActor->AttributeComponent;
 	SourceComponent->GetSet<USkillSystemTestAttributeSet>()->Health = StartHealth;
 	DestComponent->GetSet<USkillSystemTestAttributeSet>()->Health = StartHealth;
 
@@ -5455,8 +5448,8 @@ bool GameplayEffectsTest_ChanceToExecuteOnIncomingGEImmunity(UWorld *World, FAut
 
 	UProperty *DamageProperty = FindFieldChecked<UProperty>(USkillSystemTestAttributeSet::StaticClass(), GET_MEMBER_NAME_CHECKED(USkillSystemTestAttributeSet, Damage));
 
-	UAttributeComponent * SourceComponent = Cast<UAttributeComponent>(SourceActor->CreateComponentFromTemplate(UAttributeComponent::StaticClass()->GetDefaultObject<UAttributeComponent>(), FString(TEXT("AttributeComponent"))));
-	UAttributeComponent * DestComponent = Cast<UAttributeComponent>(DestActor->CreateComponentFromTemplate(UAttributeComponent::StaticClass()->GetDefaultObject<UAttributeComponent>(), FString(TEXT("AttributeComponent"))));
+	UAttributeComponent * SourceComponent = SourceActor->AttributeComponent;
+	UAttributeComponent * DestComponent = DestActor->AttributeComponent;
 	SourceComponent->GetSet<USkillSystemTestAttributeSet>()->Health = StartHealth;
 	DestComponent->GetSet<USkillSystemTestAttributeSet>()->Health = StartHealth;
 
@@ -5547,8 +5540,8 @@ bool GameplayEffectsTest_ModifyChanceToApplyToTarget(UWorld *World, FAutomationT
 
 	UProperty *DamageProperty = FindFieldChecked<UProperty>(USkillSystemTestAttributeSet::StaticClass(), GET_MEMBER_NAME_CHECKED(USkillSystemTestAttributeSet, Damage));
 
-	UAttributeComponent * SourceComponent = Cast<UAttributeComponent>(SourceActor->CreateComponentFromTemplate(UAttributeComponent::StaticClass()->GetDefaultObject<UAttributeComponent>(), FString(TEXT("AttributeComponent"))));
-	UAttributeComponent * DestComponent = Cast<UAttributeComponent>(DestActor->CreateComponentFromTemplate(UAttributeComponent::StaticClass()->GetDefaultObject<UAttributeComponent>(), FString(TEXT("AttributeComponent"))));
+	UAttributeComponent * SourceComponent = SourceActor->AttributeComponent;
+	UAttributeComponent * DestComponent = DestActor->AttributeComponent;
 	SourceComponent->GetSet<USkillSystemTestAttributeSet>()->Health = StartHealth;
 	DestComponent->GetSet<USkillSystemTestAttributeSet>()->Health = StartHealth;
 
@@ -5605,8 +5598,8 @@ bool GameplayEffectsTest_ModifyChanceToExecuteOnGE(UWorld *World, FAutomationTes
 
 	UProperty *DamageProperty = FindFieldChecked<UProperty>(USkillSystemTestAttributeSet::StaticClass(), GET_MEMBER_NAME_CHECKED(USkillSystemTestAttributeSet, Damage));
 
-	UAttributeComponent * SourceComponent = Cast<UAttributeComponent>(SourceActor->CreateComponentFromTemplate(UAttributeComponent::StaticClass()->GetDefaultObject<UAttributeComponent>(), FString(TEXT("AttributeComponent"))));
-	UAttributeComponent * DestComponent = Cast<UAttributeComponent>(DestActor->CreateComponentFromTemplate(UAttributeComponent::StaticClass()->GetDefaultObject<UAttributeComponent>(), FString(TEXT("AttributeComponent"))));
+	UAttributeComponent * SourceComponent = SourceActor->AttributeComponent;
+	UAttributeComponent * DestComponent = DestActor->AttributeComponent;
 	SourceComponent->GetSet<USkillSystemTestAttributeSet>()->Health = StartHealth;
 	DestComponent->GetSet<USkillSystemTestAttributeSet>()->Health = StartHealth;
 
diff --git a/Engine/Source/Runtime/SkillSystem/Private/PlayMontageAndWaitAction.cpp b/Engine/Source/Runtime/SkillSystem/Private/PlayMontageAndWaitAction.cpp
deleted file mode 100644
diff --git a/Engine/Source/Runtime/SkillSystem/Private/SkillSystemBlueprintLibrary.cpp b/Engine/Source/Runtime/SkillSystem/Private/SkillSystemBlueprintLibrary.cpp
--- a/Engine/Source/Runtime/SkillSystem/Private/SkillSystemBlueprintLibrary.cpp
+++ b/Engine/Source/Runtime/SkillSystem/Private/SkillSystemBlueprintLibrary.cpp
@@ -6,6 +6,7 @@
 #include "Abilities/GameplayAbility_Instanced.h"
 #include "Abilities/Tasks/AbilityTask_PlayMontageAndWait.h"
 #include "Abilities/Tasks/AbilityTask_WaitMovementModeChange.h"
+#include "Abilities/Tasks/AbilityTask_WaitOverlap.h"
 #include "LatentActions.h"
 
 USkillSystemBlueprintLibrary::USkillSystemBlueprintLibrary(const class FPostConstructInitializeProperties& PCIP)
@@ -13,40 +14,10 @@ USkillSystemBlueprintLibrary::USkillSystemBlueprintLibrary(const class FPostCons
 {
 }
 
-// Fixme: remove this function!!!
-void USkillSystemBlueprintLibrary::StaticFuncPlayMontageAndWait(AActor* WorldContextObject, UAnimMontage * Montage, FLatentActionInfo LatentInfo)
-{
-	check(WorldContextObject);
-
-	if (UWorld* World = GEngine->GetWorldFromContextObject(WorldContextObject))
-	{
-		FLatentActionManager& LatentActionManager = World->GetLatentActionManager();
-		if (LatentActionManager.FindExistingAction<FPlayMontageAndWaitAction>(LatentInfo.CallbackTarget, LatentInfo.UUID) == NULL)
-		{
-			FGameplayAbilityActorInfo	ActorInfo;
-
-			AActor * ActorOwner = NULL;
-			UObject * TestObj = WorldContextObject;
-			while(ActorOwner == NULL && TestObj != NULL)
-			{
-				ActorOwner = Cast<AActor>(TestObj);
-				TestObj = TestObj->GetOuter();
-			}
-
-			if (ActorOwner)
-			{
-				ActorInfo.InitFromActor(ActorOwner);  
-				LatentActionManager.AddNewAction(LatentInfo.CallbackTarget, LatentInfo.UUID, new FPlayMontageAndWaitAction(ActorInfo, Montage, LatentInfo));
-			}
-		}
-	}
-}
-
 UAbilityTask_PlayMontageAndWait* USkillSystemBlueprintLibrary::CreatePlayMontageAndWaitProxy(class UObject* WorldContextObject, UAnimMontage *MontageToPlay)
 {
 	check(WorldContextObject);
-
-	UGameplayAbility_Instanced * Ability = Cast<UGameplayAbility_Instanced>(WorldContextObject);
+	UGameplayAbility_Instanced * Ability = CastChecked<UGameplayAbility_Instanced>(WorldContextObject);
 	if (Ability)
 	{
 		AActor * ActorOwner = Cast<AActor>(Ability->GetOuter());
@@ -73,8 +44,7 @@ UAbilityTask_PlayMontageAndWait* USkillSystemBlueprintLibrary::CreatePlayMontage
 UAbilityTask_WaitMovementModeChange* USkillSystemBlueprintLibrary::CreateWaitMovementModeChange(class UObject* WorldContextObject, EMovementMode NewMode)
 {
 	check(WorldContextObject);
-
-	UGameplayAbility_Instanced * Ability = Cast<UGameplayAbility_Instanced>(WorldContextObject);
+	UGameplayAbility_Instanced * Ability = CastChecked<UGameplayAbility_Instanced>(WorldContextObject);
 	if (Ability)
 	{
 		AActor * ActorOwner = Cast<AActor>(Ability->GetOuter());
@@ -95,4 +65,57 @@ UAbilityTask_WaitMovementModeChange* USkillSystemBlueprintLibrary::CreateWaitMov
 	}
 
 	return NULL;
+}
+
+/**
+ *	Need:
+ *	-Easy way to specify which primitive components should be used for this overlap test
+ *	-Easy way to specify which types of actors/collision overlaps that we care about/want to block on
+ */
+
+UAbilityTask_WaitOverlap* USkillSystemBlueprintLibrary::CreateWaitOverlap(class UObject* WorldContextObject, EMovementMode NewMode)
+{
+	check(WorldContextObject);
+	UGameplayAbility_Instanced * Ability = CastChecked<UGameplayAbility_Instanced>(WorldContextObject);
+	if (Ability)
+	{
+		AActor * ActorOwner = Cast<AActor>(Ability->GetOuter());
+
+		UAbilityTask_WaitOverlap * MyObj = NULL;
+		MyObj = NewObject<UAbilityTask_WaitOverlap>();
+
+
+		// TEMP
+		UPrimitiveComponent * PrimComponent = Cast<UPrimitiveComponent>(ActorOwner->GetRootComponent());
+		if (!PrimComponent)
+		{			
+			PrimComponent = ActorOwner->FindComponentByClass<UPrimitiveComponent>();
+		}
+		check(PrimComponent);
+
+		PrimComponent->OnComponentBeginOverlap.AddDynamic(MyObj, &UAbilityTask_WaitOverlap::OnOverlapCallback);
+		PrimComponent->OnComponentHit.AddDynamic(MyObj, &UAbilityTask_WaitOverlap::OnHitCallback);
+
+		return MyObj;
+	}
+
+	return NULL;
+}
+
+/**
+ *	Fixme: This is supposed to be a small optimization but here we are going through module code, then a global uobject, then finally a virtual function - yuck!
+ *		-think of a better way to have a global, overridable per project function. Thats it!
+ */
+
+UAttributeComponent* USkillSystemBlueprintLibrary::GetAttributeComponent(AActor *Actor)
+{
+	return ISkillSystemModule::Get().GetSkillSystemGlobals().GetAttributeComponentFromActor(Actor);
+}
+
+void USkillSystemBlueprintLibrary::ApplyGameplayEffectToTargetData(FGameplayAbilityTargetDataHandle Target, UGameplayEffect *GameplayEffect, const FGameplayAbilityActorInfo InstigatorInfo)
+{
+	if (Target.Data.IsValid())
+	{
+		Target.Data->ApplyGameplayEffect(GameplayEffect, InstigatorInfo);	
+	}
 }
\ No newline at end of file
diff --git a/Engine/Source/Runtime/SkillSystem/Private/SkillSystemDebugHUD.cpp b/Engine/Source/Runtime/SkillSystem/Private/SkillSystemDebugHUD.cpp
--- a/Engine/Source/Runtime/SkillSystem/Private/SkillSystemDebugHUD.cpp
+++ b/Engine/Source/Runtime/SkillSystem/Private/SkillSystemDebugHUD.cpp
@@ -80,7 +80,10 @@ void ASkillSystemDebugHUD::DrawDebugAttributeComponent(UAttributeComponent *Comp
 	float X = 20.f;
 	float Y = 20.f;
 
-	FString String = FString::Printf(TEXT("%s (%d)"), *Component->GetPathName(), Component->IsDefaultSubobject());
+	FString String = FString::Printf(TEXT("%.2f"), Component->GetWorld()->GetTimeSeconds());
+	DrawWithBackground(Font, String, Color, EAlignHorizontal::Left, X, EAlignVertical::Top, Y);
+
+	String = FString::Printf(TEXT("%s (%d)"), *Component->GetPathName(), Component->IsDefaultSubobject());
 	DrawWithBackground(Font, String, Color, EAlignHorizontal::Left, X, EAlignVertical::Top, Y);
 
 	
@@ -151,6 +154,7 @@ static void	ToggleDebugHUD(const TArray<FString>& Args, UWorld* InWorld)
 	{
 		FDebugDrawDelegate DrawDebugDelegate = FDebugDrawDelegate::CreateUObject(HUD, &ASkillSystemDebugHUD::DrawDebugHUD);
 		UDebugDrawService::Unregister(DrawDebugDelegate);
+		HUD->Destroy();
 	}
 }
 
diff --git a/Engine/Source/Runtime/SkillSystem/Private/SkillSystemGlobals.cpp b/Engine/Source/Runtime/SkillSystem/Private/SkillSystemGlobals.cpp
--- a/Engine/Source/Runtime/SkillSystem/Private/SkillSystemGlobals.cpp
+++ b/Engine/Source/Runtime/SkillSystem/Private/SkillSystemGlobals.cpp
@@ -59,3 +59,16 @@ void USkillSystemGlobals::OnDataTableReimported(UObject* InObject)
 }
 
 #endif
+
+
+UAttributeComponent * USkillSystemGlobals::GetAttributeComponentFromActor(AActor *Actor) const
+{
+	if (Actor)
+	{
+		return Actor->FindComponentByClass<UAttributeComponent>();
+	}
+
+	// Caller should check this
+	ensure(false);
+	return NULL;
+}
\ No newline at end of file
diff --git a/Engine/Source/Runtime/SkillSystem/Private/SkillSystemModulePrivatePCH.h b/Engine/Source/Runtime/SkillSystem/Private/SkillSystemModulePrivatePCH.h
--- a/Engine/Source/Runtime/SkillSystem/Private/SkillSystemModulePrivatePCH.h
+++ b/Engine/Source/Runtime/SkillSystem/Private/SkillSystemModulePrivatePCH.h
@@ -16,10 +16,6 @@
 
 #include "SkillSystemTypes.h"
 
-//#include "SkillSystemClasses.h"
-
 #include "SkillSystem.h"
 #include "SkillSystemModule.h"
-#include "SkillSystemLog.h"
-
-#include "PlayMontageAndWaitAction.h"
+#include "SkillSystemLog.h"
\ No newline at end of file
diff --git a/Engine/Source/Runtime/SkillSystem/Private/SkillSystemTestAttributeSet.cpp b/Engine/Source/Runtime/SkillSystem/Private/SkillSystemTestAttributeSet.cpp
--- a/Engine/Source/Runtime/SkillSystem/Private/SkillSystemTestAttributeSet.cpp
+++ b/Engine/Source/Runtime/SkillSystem/Private/SkillSystemTestAttributeSet.cpp
@@ -56,7 +56,7 @@ void USkillSystemTestAttributeSet::PreAttributeModify(struct FGameplayEffectModC
 		if (Data.EvaluatedData.Magnitude > 0.f)
 		{
 			// Check the source - does he have Crit?
-			USkillSystemTestAttributeSet * SourceAttributes = Data.EffectSpec.InstigatorStack.GetOriginInstigatorAttributeComponent()->GetSet<USkillSystemTestAttributeSet>();
+			USkillSystemTestAttributeSet * SourceAttributes = Data.EffectSpec.InstigatorContext.GetOriginalInstigatorAttributeComponent()->GetSet<USkillSystemTestAttributeSet>();
 			if (SourceAttributes && SourceAttributes->CritChance > 0.f)
 			{
 				if (FMath::FRand() <= SourceAttributes->CritChance)
diff --git a/Engine/Source/Runtime/SkillSystem/Private/SkillSystemTestPawn.cpp b/Engine/Source/Runtime/SkillSystem/Private/SkillSystemTestPawn.cpp
--- a/Engine/Source/Runtime/SkillSystem/Private/SkillSystemTestPawn.cpp
+++ b/Engine/Source/Runtime/SkillSystem/Private/SkillSystemTestPawn.cpp
@@ -32,26 +32,26 @@ void ASkillSystemTestPawn::PostInitializeComponents()
 	*/
 }
 
-void ASkillSystemTestPawn::GameplayCueActivated(const FGameplayTagContainer & GameplayCueTags, float NormalizedMagnitude)
+void ASkillSystemTestPawn::GameplayCueActivated(const FGameplayTagContainer & GameplayCueTags, float NormalizedMagnitude, const FGameplayEffectInstigatorContext InstigatorContext)
 {
 	SKILL_LOG(Log, TEXT("GameplayCueExecuted: %s. %.2f"), *GetName(), NormalizedMagnitude);
-	GameplayCueHandler.GameplayCueActivated(GameplayCueTags, NormalizedMagnitude);
+	GameplayCueHandler.GameplayCueActivated(GameplayCueTags, NormalizedMagnitude, InstigatorContext);
 }
 
-void ASkillSystemTestPawn::GameplayCueExecuted(const FGameplayTagContainer & GameplayCueTags, float NormalizedMagnitude)
+void ASkillSystemTestPawn::GameplayCueExecuted(const FGameplayTagContainer & GameplayCueTags, float NormalizedMagnitude, const FGameplayEffectInstigatorContext InstigatorContext)
 {
 	SKILL_LOG(Log, TEXT("GameplayCueExecuted: %s. %.2f"), *GetName(), NormalizedMagnitude);
-	GameplayCueHandler.GameplayCueExecuted(GameplayCueTags, NormalizedMagnitude);
+	GameplayCueHandler.GameplayCueExecuted(GameplayCueTags, NormalizedMagnitude, InstigatorContext);
 }
 
-void ASkillSystemTestPawn::GameplayCueAdded(const FGameplayTagContainer & GameplayCueTags, float NormalizedMagnitude)
+void ASkillSystemTestPawn::GameplayCueAdded(const FGameplayTagContainer & GameplayCueTags, float NormalizedMagnitude, const FGameplayEffectInstigatorContext InstigatorContext)
 {
 	SKILL_LOG(Log, TEXT("GameplayCueAdded: %s. %.2f"), *GetName(), NormalizedMagnitude);
-	GameplayCueHandler.GameplayCueAdded(GameplayCueTags, NormalizedMagnitude);
+	GameplayCueHandler.GameplayCueAdded(GameplayCueTags, NormalizedMagnitude, InstigatorContext);
 }
 
-void ASkillSystemTestPawn::GameplayCueRemoved(const FGameplayTagContainer & GameplayCueTags, float NormalizedMagnitude)
+void ASkillSystemTestPawn::GameplayCueRemoved(const FGameplayTagContainer & GameplayCueTags, float NormalizedMagnitude, const FGameplayEffectInstigatorContext InstigatorContext)
 {
 	SKILL_LOG(Log, TEXT("GameplayCueRemoved: %s. %.2f"), *GetName(), NormalizedMagnitude);
-	GameplayCueHandler.GameplayCueRemoved(GameplayCueTags, NormalizedMagnitude);
+	GameplayCueHandler.GameplayCueRemoved(GameplayCueTags, NormalizedMagnitude, InstigatorContext);
 }
\ No newline at end of file
diff --git a/Engine/Source/Runtime/SkillSystem/Public/PlayMontageAndWaitAction.h b/Engine/Source/Runtime/SkillSystem/Public/PlayMontageAndWaitAction.h
deleted file mode 100644
diff --git a/Engine/Source/Runtime/SkillSystem/Public/SkillSystemTypes.h b/Engine/Source/Runtime/SkillSystem/Public/SkillSystemTypes.h
--- a/Engine/Source/Runtime/SkillSystem/Public/SkillSystemTypes.h
+++ b/Engine/Source/Runtime/SkillSystem/Public/SkillSystemTypes.h
@@ -5,6 +5,7 @@
 #include "SkillSystemTypes.generated.h"
 
 class UAttributeComponent;
+class UGameplayEffect;
 
 UENUM(BlueprintType)
 namespace EGameplayAbilityActivationMode
@@ -64,3 +65,116 @@ struct SKILLSYSTEM_API FGameplayAbilityActorInfo
 
 	void SetActivationConfirmed();
 };
+
+
+// ----------------------------------------------------
+
+/**
+ *	A generic structure for targeting data. We want generic functions to produce this data and other generic
+ *	functions to consume this data.
+ *	
+ *	We expect this to be able to hold specific actors/object reference and also generic location/direction/origin 
+ *	information.
+ *	
+ *	Some example producers:
+ *		-Overlap/Hit collision event generates TargetData about who was hit in a melee attack
+ *		-A mouse input causes a hit trace and the actor infront of the crosshair is turned into TargetData
+ *		-A mouse input causes TargetData to be generated from the owner's crosshair view origin/direction
+ *		-An AOE/aura pulses and all actors in a radius around the instigator are added to TargetData
+ *		-Panzer Dragoon style 'painting' targeting mode
+ *		-MMORPG style ground AOE targeting style (potentially both a location on the ground and actors that were targeted)
+ *		
+ *	Some example consumers:
+ *		-Apply a GameplayEffect to all actors in TargetData
+ *		-Find closest actor from all in TargetData
+ *		-Call some function on all actors in TargetData
+ *		-Filter or merge TargetDatas
+ *		-Spawn a new actor at a TargetData location
+ *		
+ *		
+ *		
+ *	Maybe it is better to distinguish between actor list targeting vs positional targeting data?
+ *		-AOE/aura type of targeting data blurs the line
+ *		
+ *		
+ */
+
+struct SKILLSYSTEM_API FGameplayAbilityTargetData
+{
+	void ApplyGameplayEffect(UGameplayEffect *GameplayEffect, const FGameplayAbilityActorInfo InstigatorInfo);
+
+	virtual TArray<AActor*>	GetActors() const = 0;
+	
+	virtual bool HasHitResult() const 
+	{	
+		return false;
+	}
+
+	virtual const FHitResult * GetHitResult() const
+	{
+		return NULL;
+	}
+
+	virtual FString ToString() const;
+};
+
+struct SKILLSYSTEM_API FGameplayAbilityTargetData_SingleTargetHit : public FGameplayAbilityTargetData
+{
+	FGameplayAbilityTargetData_SingleTargetHit(const FHitResult InHitResult)
+	: HitResult(InHitResult)
+	{ }
+
+	virtual TArray<AActor*>	GetActors() const
+	{
+		TArray<AActor*>	Actors;
+		if (HitResult.Actor.IsValid())
+		{
+			Actors.Push(HitResult.Actor.Get());
+		}
+		return Actors;
+	}
+
+	virtual bool HasHitResult() const
+	{
+		return true;
+	}
+
+	virtual const FHitResult * GetHitResult() const
+	{
+		return &HitResult;
+	}
+
+	FHitResult	HitResult;
+};
+
+/**
+ *	Handle for Targeting Data. This servers two main purposes:
+ *		-Avoid us having to copy around the full targeting data structure in Blueprints
+ *		-Allows us to leverage polymorphism in the target data structure
+ *		-Allows us to implement NetSerialize and replicate by value between clients/server	
+ *		
+ *		-Avoid using UObjects could be used to give us polymorphism and by reference passing in blueprints.
+ *		-However we would still be screwed when it came to replication
+ *		
+ *		-Replication by value
+ *		-Pass by reference in blueprints
+ *		-Polymophism in TargetData structure
+ *		
+ */
+
+USTRUCT(BlueprintType)
+struct SKILLSYSTEM_API FGameplayAbilityTargetDataHandle
+{
+	GENERATED_USTRUCT_BODY()
+	
+	TSharedPtr<FGameplayAbilityTargetData>	Data;
+};
+
+template<>
+struct TStructOpsTypeTraits<FGameplayAbilityTargetDataHandle> : public TStructOpsTypeTraitsBase
+{
+	enum
+	{
+		WithCopy = true		// Necessary so that TSharedPtr<FGameplayAbilityTargetData> Data is copied around
+	};
+};
