Fix SkillSystem in non-unity

diff --git a/Engine/Source/Runtime/SkillSystem/Classes/GameplayCueView.h b/Engine/Source/Runtime/SkillSystem/Classes/GameplayCueView.h
--- a/Engine/Source/Runtime/SkillSystem/Classes/GameplayCueView.h
+++ b/Engine/Source/Runtime/SkillSystem/Classes/GameplayCueView.h
@@ -3,6 +3,7 @@
 #pragma once
 
 #include "GameplayTags.h"
+#include "GameplayEffect.h"
 #include "GameplayCueView.generated.h"
 
 /** 
diff --git a/Engine/Source/Runtime/SkillSystem/Private/Abilities/GameplayAbility.cpp b/Engine/Source/Runtime/SkillSystem/Private/Abilities/GameplayAbility.cpp
--- a/Engine/Source/Runtime/SkillSystem/Private/Abilities/GameplayAbility.cpp
+++ b/Engine/Source/Runtime/SkillSystem/Private/Abilities/GameplayAbility.cpp
@@ -4,6 +4,7 @@
 #include "Abilities/GameplayAbility.h"
 #include "GameplayEffect.h"
 #include "AttributeComponent.h"
+#include "SkillSystemBlueprintLibrary.h"
 
 // --------------------------------------------------------------------------------------------------------------------------------------------------------
 //
diff --git a/Engine/Source/Runtime/SkillSystem/Private/SkillSystemGlobals.cpp b/Engine/Source/Runtime/SkillSystem/Private/SkillSystemGlobals.cpp
--- a/Engine/Source/Runtime/SkillSystem/Private/SkillSystemGlobals.cpp
+++ b/Engine/Source/Runtime/SkillSystem/Private/SkillSystemGlobals.cpp
@@ -1,6 +1,7 @@
 // Copyright 1998-2014 Epic Games, Inc. All Rights Reserved.
 
 #include "SkillSystemModulePrivatePCH.h"
+#include "AttributeComponent.h"
 
 USkillSystemGlobals::USkillSystemGlobals(const class FPostConstructInitializeProperties& PCIP)
 : Super(PCIP)
