Remove prediction key from gameplay effect tests.

diff --git a/Engine/Source/Runtime/SkillSystem/Private/AttributeComponent.cpp b/Engine/Source/Runtime/SkillSystem/Private/AttributeComponent.cpp
--- a/Engine/Source/Runtime/SkillSystem/Private/AttributeComponent.cpp
+++ b/Engine/Source/Runtime/SkillSystem/Private/AttributeComponent.cpp
@@ -253,7 +253,7 @@ FActiveGameplayEffectHandle UAttributeComponent::ApplyGameplayEffectSpecToSelf(c
 {
 	// Temp, only non instant, non periodic GEs can be predictive
 	// Effects with other effects may be a mix so go with non-predictive
-	check(!(BaseQualifier.PredictionKey() > 0) || ((Spec.GetDuration() != UGameplayEffect::INSTANT_APPLICATION && Spec.GetPeriod() == UGameplayEffect::NO_PERIOD) || Spec.TargetEffectSpecs.Num() > 0));
+	check((BaseQualifier.PredictionKey() == 0) || (Spec.GetDuration() != UGameplayEffect::INSTANT_APPLICATION && Spec.GetPeriod() == UGameplayEffect::NO_PERIOD));
 
 	// check if the effect being applied actually succeeds
 	float ChanceToApply = Spec.GetChanceToApplyToTarget();
diff --git a/Engine/Source/Runtime/SkillSystem/Private/GameplayEffectTests.cpp b/Engine/Source/Runtime/SkillSystem/Private/GameplayEffectTests.cpp
--- a/Engine/Source/Runtime/SkillSystem/Private/GameplayEffectTests.cpp
+++ b/Engine/Source/Runtime/SkillSystem/Private/GameplayEffectTests.cpp
@@ -12,17 +12,10 @@
 #include "GameplayEffectStackingExtension_CappedNumberTest.h"
 #include "GameplayEffectStackingExtension_DiminishingReturnsTest.h"
 
-//#include "SkillSystemTestPawn.h"
-
 IMPLEMENT_SIMPLE_AUTOMATION_TEST(FGameplayEffectsTest, "SkillSystem.GameplayEffects", EAutomationTestFlags::ATF_Editor)
 
 #define SKILL_TEST_TEXT( Format, ... ) FString::Printf(TEXT("%s - %d: %s"), TEXT(__FILE__) , __LINE__ , *FString::Printf(TEXT(Format), ##__VA_ARGS__) )
 
-namespace
-{
-	const int32 TEST_PREDICTIONKEY = 1;
-}
-
 void GameplayTest_TickWorld(UWorld *World, float Time)
 {
 	const float step = 0.1f;
@@ -228,11 +221,8 @@ bool GameplayEffectsTest_InstantDamage_Buffed(UWorld *World, FAutomationTestBase
 		BaseDmgEffect->Modifiers[0].OwnedTags.AddTag(IGameplayTagsModule::RequestGameplayTag(FName(TEXT("Damage.Basic"))));
 		BaseDmgEffect->Duration.SetValue(UGameplayEffect::INFINITE_DURATION);
 
-		FModifierQualifier Qualifier = FModifierQualifier();
-		Qualifier.PredictionKey(TEST_PREDICTIONKEY);
-
 		// Apply to self
-		SourceComponent->ApplyGameplayEffectToTarget(BaseDmgEffect, SourceComponent, 1.f, Qualifier);
+		SourceComponent->ApplyGameplayEffectToTarget(BaseDmgEffect, SourceComponent, 1.f);
 	}
 
 	// Apply Damage
@@ -297,11 +287,8 @@ bool GameplayEffectsTest_TemporaryDamage(UWorld *World, FAutomationTestBase * Te
 		BaseDmgEffect->Modifiers[0].OwnedTags.AddTag(IGameplayTagsModule::RequestGameplayTag(FName(TEXT("Damage.Basic"))));
 		BaseDmgEffect->Duration.Value = UGameplayEffect::INFINITE_DURATION;
 
-		FModifierQualifier Qualifier = FModifierQualifier();
-		Qualifier.PredictionKey(TEST_PREDICTIONKEY);
-
 		// Apply to target
-		AppliedHandle = SourceComponent->ApplyGameplayEffectToTarget(BaseDmgEffect, DestComponent, 1.f, Qualifier);
+		AppliedHandle = SourceComponent->ApplyGameplayEffectToTarget(BaseDmgEffect, DestComponent, 1.f);
 
 		float ExpectedValue = (StartHealth + DamageValue);
 
@@ -359,11 +346,8 @@ bool GameplayEffectsTest_TemporaryDamageBuffed(UWorld *World, FAutomationTestBas
 		BaseDmgEffect->Modifiers[0].OwnedTags.AddTag(IGameplayTagsModule::RequestGameplayTag(FName(TEXT("Damage.Basic"))));
 		BaseDmgEffect->Duration.Value = UGameplayEffect::INFINITE_DURATION;
 
-		FModifierQualifier Qualifier = FModifierQualifier();
-		Qualifier.PredictionKey(TEST_PREDICTIONKEY);
-
 		// Apply to target
-		AppliedHandle = SourceComponent->ApplyGameplayEffectToTarget(BaseDmgEffect, DestComponent, 1.f, Qualifier);
+		AppliedHandle = SourceComponent->ApplyGameplayEffectToTarget(BaseDmgEffect, DestComponent, 1.f);
 
 		float ExpectedValue = (StartHealth + DamageValue);
 
@@ -454,11 +438,8 @@ bool GameplayEffectsTest_TemporaryDamageTemporaryBuff(UWorld *World, FAutomation
 		BaseDmgEffect->Modifiers[0].OwnedTags.AddTag(IGameplayTagsModule::RequestGameplayTag(FName(TEXT("Damage.Basic"))));
 		BaseDmgEffect->Duration.Value = UGameplayEffect::INFINITE_DURATION;
 
-		FModifierQualifier Qualifier = FModifierQualifier();
-		Qualifier.PredictionKey(TEST_PREDICTIONKEY);
-
 		// Apply to target
-		AppliedHandle = SourceComponent->ApplyGameplayEffectToTarget(BaseDmgEffect, DestComponent, 1.f, Qualifier);
+		AppliedHandle = SourceComponent->ApplyGameplayEffectToTarget(BaseDmgEffect, DestComponent, 1.f);
 
 		float ExpectedValue = (StartHealth + DamageValue);
 
@@ -482,11 +463,8 @@ bool GameplayEffectsTest_TemporaryDamageTemporaryBuff(UWorld *World, FAutomation
 		BuffDmgEffect->Duration.Value = UGameplayEffect::INFINITE_DURATION;
 		BuffDmgEffect->CopyPolicy = EGameplayEffectCopyPolicy::AlwaysLink;		// Force this to link, so that when we remove it it will go away to any modifier it was applied to
 
-		FModifierQualifier Qualifier = FModifierQualifier();
-		Qualifier.PredictionKey(TEST_PREDICTIONKEY);
-
 		// Apply to target
-		AppliedHandleBuff = SourceComponent->ApplyGameplayEffectToTarget(BuffDmgEffect, DestComponent, 1.f, Qualifier);
+		AppliedHandleBuff = SourceComponent->ApplyGameplayEffectToTarget(BuffDmgEffect, DestComponent, 1.f);
 
 		float ExpectedValue = (StartHealth + (DamageValue * DamageBuffMultiplier));
 
@@ -569,11 +547,8 @@ bool GameplayEffectsTest_LinkedBuffDestroy(UWorld *World, FAutomationTestBase *
 		BuffEffect->Duration.Value = UGameplayEffect::INFINITE_DURATION;
 		BuffEffect->CopyPolicy = EGameplayEffectCopyPolicy::AlwaysLink;		// Always link so that when this is destroyed, health return to normal
 
-		FModifierQualifier Qualifier = FModifierQualifier();
-		Qualifier.PredictionKey(TEST_PREDICTIONKEY);
-
 		// Apply to target
-		AppliedHandleBuff = SourceComponent->ApplyGameplayEffectToTarget(BuffEffect, SourceComponent, 1.f, Qualifier);
+		AppliedHandleBuff = SourceComponent->ApplyGameplayEffectToTarget(BuffEffect, SourceComponent, 1.f);
 
 		Test->TestTrue(SKILL_TEST_TEXT("Number of Source GameplayEffect: %d == 1", SourceComponent->GetNumActiveGameplayEffect()), SourceComponent->GetNumActiveGameplayEffect() == 1);
 	}
@@ -590,11 +565,8 @@ bool GameplayEffectsTest_LinkedBuffDestroy(UWorld *World, FAutomationTestBase *
 		BaseDmgEffect->Modifiers[0].OwnedTags.AddTag(IGameplayTagsModule::RequestGameplayTag(FName(TEXT("Damage.Basic"))));
 		BaseDmgEffect->Duration.Value = UGameplayEffect::INFINITE_DURATION;
 
-		FModifierQualifier Qualifier = FModifierQualifier();
-		Qualifier.PredictionKey(TEST_PREDICTIONKEY);
-
 		// Apply to target
-		AppliedHandle = SourceComponent->ApplyGameplayEffectToTarget(BaseDmgEffect, DestComponent, 1.f, Qualifier);
+		AppliedHandle = SourceComponent->ApplyGameplayEffectToTarget(BaseDmgEffect, DestComponent, 1.f);
 
 		float ExpectedValue = (StartHealth + (DamageValue * DamageBuffMultiplier));
 		float ActualValue = DestComponent->GetSet<USkillSystemTestAttributeSet>()->Health;
@@ -663,11 +635,8 @@ bool GameplayEffectsTest_SnapshotBuffDestroy(UWorld *World, FAutomationTestBase
 		BuffEffect->Duration.Value = UGameplayEffect::INFINITE_DURATION;
 		BuffEffect->CopyPolicy = EGameplayEffectCopyPolicy::AlwaysSnapshot;		// Always snapshot (though the default for outgoing should already be snapshot - but this could change per project at some point)
 
-		FModifierQualifier Qualifier = FModifierQualifier();
-		Qualifier.PredictionKey(TEST_PREDICTIONKEY);
-
 		// Apply to target
-		AppliedHandleBuff = SourceComponent->ApplyGameplayEffectToTarget(BuffEffect, SourceComponent, 1.f, Qualifier);
+		AppliedHandleBuff = SourceComponent->ApplyGameplayEffectToTarget(BuffEffect, SourceComponent, 1.f);
 
 		Test->TestTrue(SKILL_TEST_TEXT("Number of Source GameplayEffect: %d == 1", SourceComponent->GetNumActiveGameplayEffect()), SourceComponent->GetNumActiveGameplayEffect() == 1);
 	}
@@ -684,11 +653,8 @@ bool GameplayEffectsTest_SnapshotBuffDestroy(UWorld *World, FAutomationTestBase
 		BaseDmgEffect->Modifiers[0].OwnedTags.AddTag(IGameplayTagsModule::RequestGameplayTag(FName(TEXT("Damage.Basic"))));
 		BaseDmgEffect->Duration.Value = UGameplayEffect::INFINITE_DURATION;
 
-		FModifierQualifier Qualifier = FModifierQualifier();
-		Qualifier.PredictionKey(TEST_PREDICTIONKEY);
-
 		// Apply to target
-		AppliedHandle = SourceComponent->ApplyGameplayEffectToTarget(BaseDmgEffect, DestComponent, 1.f, Qualifier);
+		AppliedHandle = SourceComponent->ApplyGameplayEffectToTarget(BaseDmgEffect, DestComponent, 1.f);
 
 		float ExpectedValue = (StartHealth + (DamageValue * DamageBuffMultiplier));
 		float ActualValue = DestComponent->GetSet<USkillSystemTestAttributeSet>()->Health;
@@ -763,11 +729,8 @@ bool GameplayEffectsTest_DurationBuff(UWorld *World, FAutomationTestBase * Test)
 		BaseDmgEffect->Modifiers[0].OwnedTags.AddTag(IGameplayTagsModule::RequestGameplayTag(FName(TEXT("Damage.Basic"))));
 		BaseDmgEffect->Duration.SetValue(BaseDuration);
 
-		FModifierQualifier Qualifier = FModifierQualifier();
-		Qualifier.PredictionKey(TEST_PREDICTIONKEY);
-
 		// Apply to target
-		AppliedDamageHandle = SourceComponent->ApplyGameplayEffectToTarget(BaseDmgEffect, DestComponent, 1.f, Qualifier);
+		AppliedDamageHandle = SourceComponent->ApplyGameplayEffectToTarget(BaseDmgEffect, DestComponent, 1.f);
 
 		float ExpectedHealthValue = (StartHealth + (DamageValue));
 		float ActualHealthValue = DestComponent->GetSet<USkillSystemTestAttributeSet>()->Health;
@@ -796,11 +759,8 @@ bool GameplayEffectsTest_DurationBuff(UWorld *World, FAutomationTestBase * Test)
 		DurationEffect->Modifiers[0].OwnedTags.AddTag(IGameplayTagsModule::RequestGameplayTag(FName(TEXT("Damage.Basic"))));
 		DurationEffect->Duration.SetValue(UGameplayEffect::INFINITE_DURATION);
 
-		FModifierQualifier Qualifier = FModifierQualifier();
-		Qualifier.PredictionKey(TEST_PREDICTIONKEY);
-
 		// Apply to target
-		AppliedDurationHandle = SourceComponent->ApplyGameplayEffectToTarget(DurationEffect, DestComponent, 1.f, Qualifier);
+		AppliedDurationHandle = SourceComponent->ApplyGameplayEffectToTarget(DurationEffect, DestComponent, 1.f);
 
 		float ExpectedDuration = BaseDuration + DurationDebuff;
 		float ActualDuration = DestComponent->GetGameplayEffectDuration(AppliedDamageHandle);
@@ -834,11 +794,8 @@ bool GameplayEffectsTest_DurationBuff(UWorld *World, FAutomationTestBase * Test)
 		BaseDmgEffect->Modifiers[0].OwnedTags.AddTag(IGameplayTagsModule::RequestGameplayTag(FName(TEXT("Damage.Basic"))));
 		BaseDmgEffect->Duration.SetValue(BaseDuration);
 
-		FModifierQualifier Qualifier = FModifierQualifier();
-		Qualifier.PredictionKey(TEST_PREDICTIONKEY);
-
 		// Apply to target
-		AppliedDamageHandle = SourceComponent->ApplyGameplayEffectToTarget(BaseDmgEffect, DestComponent, 1.f, Qualifier);
+		AppliedDamageHandle = SourceComponent->ApplyGameplayEffectToTarget(BaseDmgEffect, DestComponent, 1.f);
 
 		float ExpectedHealthValue = (StartHealth + (DamageValue));
 		float ActualHealthValue = DestComponent->GetSet<USkillSystemTestAttributeSet>()->Health;
@@ -867,11 +824,8 @@ bool GameplayEffectsTest_DurationBuff(UWorld *World, FAutomationTestBase * Test)
 		DurationEffect->Modifiers[0].OwnedTags.AddTag(IGameplayTagsModule::RequestGameplayTag(FName(TEXT("Damage.Basic"))));
 		DurationEffect->Duration.SetValue(UGameplayEffect::INFINITE_DURATION);
 
-		FModifierQualifier Qualifier = FModifierQualifier();
-		Qualifier.PredictionKey(TEST_PREDICTIONKEY);
-
 		// Apply to target
-		AppliedDurationHandle = SourceComponent->ApplyGameplayEffectToTarget(DurationEffect, DestComponent, 1.f, Qualifier);
+		AppliedDurationHandle = SourceComponent->ApplyGameplayEffectToTarget(DurationEffect, DestComponent, 1.f);
 
 		float ExpectedDuration = BaseDuration + DurationBuff;
 		float ActualDuration = DestComponent->GetGameplayEffectDuration(AppliedDamageHandle);
@@ -916,11 +870,8 @@ bool GameplayEffectsTest_DurationBuff(UWorld *World, FAutomationTestBase * Test)
 		BaseDmgEffect->Modifiers[0].OwnedTags.AddTag(IGameplayTagsModule::RequestGameplayTag(FName(TEXT("Damage.Basic"))));
 		BaseDmgEffect->Duration.SetValue(BaseDuration);
 
-		FModifierQualifier Qualifier = FModifierQualifier();
-		Qualifier.PredictionKey(TEST_PREDICTIONKEY);
-
 		// Apply to target
-		AppliedDamageHandle = SourceComponent->ApplyGameplayEffectToTarget(BaseDmgEffect, DestComponent, 1.f, Qualifier);
+		AppliedDamageHandle = SourceComponent->ApplyGameplayEffectToTarget(BaseDmgEffect, DestComponent, 1.f);
 
 		float ExpectedHealthValue = (StartHealth + (DamageValue));
 		float ActualHealthValue = DestComponent->GetSet<USkillSystemTestAttributeSet>()->Health;
@@ -949,11 +900,8 @@ bool GameplayEffectsTest_DurationBuff(UWorld *World, FAutomationTestBase * Test)
 		DurationEffect->Modifiers[0].OwnedTags.AddTag(IGameplayTagsModule::RequestGameplayTag(FName(TEXT("Damage.Basic"))));
 		DurationEffect->Duration.SetValue(UGameplayEffect::INFINITE_DURATION);
 
-		FModifierQualifier Qualifier = FModifierQualifier();
-		Qualifier.PredictionKey(TEST_PREDICTIONKEY);
-
 		// Apply to target
-		AppliedDurationHandle = SourceComponent->ApplyGameplayEffectToTarget(DurationEffect, DestComponent, 1.f, Qualifier);
+		AppliedDurationHandle = SourceComponent->ApplyGameplayEffectToTarget(DurationEffect, DestComponent, 1.f);
 
 		float ExpectedDuration = BaseDuration + DurationBuff;
 		float ActualDuration = DestComponent->GetGameplayEffectDuration(AppliedDamageHandle);
@@ -1036,11 +984,8 @@ bool GameplayEffectsTest_DamageBuffBuff_Basic(UWorld *World, FAutomationTestBase
 		BuffEffect->Duration.Value = UGameplayEffect::INFINITE_DURATION;
 		BuffEffect->CopyPolicy = EGameplayEffectCopyPolicy::AlwaysLink;
 
-		FModifierQualifier Qualifier = FModifierQualifier();
-		Qualifier.PredictionKey(TEST_PREDICTIONKEY);
-
 		// Apply to target
-		BuffBuffHandle = SourceComponent->ApplyGameplayEffectToTarget(BuffEffect, SourceComponent, 1.f, Qualifier);
+		BuffBuffHandle = SourceComponent->ApplyGameplayEffectToTarget(BuffEffect, SourceComponent, 1.f);
 
 		Test->TestTrue(SKILL_TEST_TEXT("Number of Source GameplayEffect: %d == 1", SourceComponent->GetNumActiveGameplayEffect()), SourceComponent->GetNumActiveGameplayEffect() == 1);
 	}
@@ -1057,11 +1002,8 @@ bool GameplayEffectsTest_DamageBuffBuff_Basic(UWorld *World, FAutomationTestBase
 		BuffEffect->Duration.Value = UGameplayEffect::INFINITE_DURATION;
 		BuffEffect->CopyPolicy = EGameplayEffectCopyPolicy::AlwaysSnapshot;
 
-		FModifierQualifier Qualifier = FModifierQualifier();
-		Qualifier.PredictionKey(TEST_PREDICTIONKEY);
-
 		// Apply to target
-		BuffHandle = SourceComponent->ApplyGameplayEffectToTarget(BuffEffect, SourceComponent, 1.f, Qualifier);
+		BuffHandle = SourceComponent->ApplyGameplayEffectToTarget(BuffEffect, SourceComponent, 1.f);
 		Test->TestTrue(SKILL_TEST_TEXT("Number of Source GameplayEffect: %d == 2", SourceComponent->GetNumActiveGameplayEffect()), SourceComponent->GetNumActiveGameplayEffect() == 2);
 
 		// Check that the buff was buffed
@@ -1083,11 +1025,8 @@ bool GameplayEffectsTest_DamageBuffBuff_Basic(UWorld *World, FAutomationTestBase
 		BaseDmgEffect->Modifiers[0].OwnedTags.AddTag(IGameplayTagsModule::RequestGameplayTag(FName(TEXT("Damage.Basic"))));
 		BaseDmgEffect->Duration.Value = UGameplayEffect::INFINITE_DURATION;
 
-		FModifierQualifier Qualifier = FModifierQualifier();
-		Qualifier.PredictionKey(TEST_PREDICTIONKEY);
-
 		// Apply to target
-		DamageHandle = SourceComponent->ApplyGameplayEffectToTarget(BaseDmgEffect, DestComponent, 1.f, Qualifier);
+		DamageHandle = SourceComponent->ApplyGameplayEffectToTarget(BaseDmgEffect, DestComponent, 1.f);
 
 		float ExpectedValue = (StartHealth + (DamageValue * (DamageBuffMultiplier + DamageBuffMultiplierBonus)));
 		float ActualValue = DestComponent->GetSet<USkillSystemTestAttributeSet>()->Health;
@@ -1186,11 +1125,8 @@ bool GameplayEffectsTest_DamageBuffBuff_FullLink(UWorld *World, FAutomationTestB
 		BuffEffect->Duration.Value = UGameplayEffect::INFINITE_DURATION;
 		BuffEffect->CopyPolicy = EGameplayEffectCopyPolicy::AlwaysLink;
 
-		FModifierQualifier Qualifier = FModifierQualifier();
-		Qualifier.PredictionKey(TEST_PREDICTIONKEY);
-
 		// Apply to target
-		BuffBuffHandle = SourceComponent->ApplyGameplayEffectToTarget(BuffEffect, SourceComponent, 1.f, Qualifier);
+		BuffBuffHandle = SourceComponent->ApplyGameplayEffectToTarget(BuffEffect, SourceComponent, 1.f);
 
 		Test->TestTrue(SKILL_TEST_TEXT("Number of Source GameplayEffect: %d == 1", SourceComponent->GetNumActiveGameplayEffect()), SourceComponent->GetNumActiveGameplayEffect() == 1);
 	}
@@ -1207,11 +1143,8 @@ bool GameplayEffectsTest_DamageBuffBuff_FullLink(UWorld *World, FAutomationTestB
 		BuffEffect->Duration.Value = UGameplayEffect::INFINITE_DURATION;
 		BuffEffect->CopyPolicy = EGameplayEffectCopyPolicy::AlwaysLink;
 
-		FModifierQualifier Qualifier = FModifierQualifier();
-		Qualifier.PredictionKey(TEST_PREDICTIONKEY);
-
 		// Apply to target
-		BuffHandle = SourceComponent->ApplyGameplayEffectToTarget(BuffEffect, SourceComponent, 1.f, Qualifier);
+		BuffHandle = SourceComponent->ApplyGameplayEffectToTarget(BuffEffect, SourceComponent, 1.f);
 		Test->TestTrue(SKILL_TEST_TEXT("Number of Source GameplayEffect: %d == 2", SourceComponent->GetNumActiveGameplayEffect()), SourceComponent->GetNumActiveGameplayEffect() == 2);
 
 		// Check that the buff was buffed
@@ -1233,11 +1166,8 @@ bool GameplayEffectsTest_DamageBuffBuff_FullLink(UWorld *World, FAutomationTestB
 		BaseDmgEffect->Modifiers[0].OwnedTags.AddTag(IGameplayTagsModule::RequestGameplayTag(FName(TEXT("Damage.Basic"))));
 		BaseDmgEffect->Duration.Value = UGameplayEffect::INFINITE_DURATION;
 
-		FModifierQualifier Qualifier = FModifierQualifier();
-		Qualifier.PredictionKey(TEST_PREDICTIONKEY);
-
 		// Apply to target
-		DamageHandle = SourceComponent->ApplyGameplayEffectToTarget(BaseDmgEffect, DestComponent, 1.f, Qualifier);
+		DamageHandle = SourceComponent->ApplyGameplayEffectToTarget(BaseDmgEffect, DestComponent, 1.f);
 
 		float ExpectedValue = (StartHealth + (DamageValue * (DamageBuffMultiplier + DamageBuffMultiplierBonus)));
 		float ActualValue = DestComponent->GetSet<USkillSystemTestAttributeSet>()->Health;
@@ -1333,11 +1263,8 @@ bool GameplayEffectsTest_DamageBuffBuff_FullSnapshot(UWorld *World, FAutomationT
 		BuffEffect->Duration.Value = UGameplayEffect::INFINITE_DURATION;
 		BuffEffect->CopyPolicy = EGameplayEffectCopyPolicy::AlwaysSnapshot;
 
-		FModifierQualifier Qualifier = FModifierQualifier();
-		Qualifier.PredictionKey(TEST_PREDICTIONKEY);
-
 		// Apply to target
-		BuffBuffHandle = SourceComponent->ApplyGameplayEffectToTarget(BuffEffect, SourceComponent, 1.f, Qualifier);
+		BuffBuffHandle = SourceComponent->ApplyGameplayEffectToTarget(BuffEffect, SourceComponent, 1.f);
 
 		Test->TestTrue(SKILL_TEST_TEXT("Number of Source GameplayEffect: %d == 1", SourceComponent->GetNumActiveGameplayEffect()), SourceComponent->GetNumActiveGameplayEffect() == 1);
 	}
@@ -1354,11 +1281,8 @@ bool GameplayEffectsTest_DamageBuffBuff_FullSnapshot(UWorld *World, FAutomationT
 		BuffEffect->Duration.Value = UGameplayEffect::INFINITE_DURATION;
 		BuffEffect->CopyPolicy = EGameplayEffectCopyPolicy::AlwaysSnapshot;
 
-		FModifierQualifier Qualifier = FModifierQualifier();
-		Qualifier.PredictionKey(TEST_PREDICTIONKEY);
-
 		// Apply to target
-		BuffHandle = SourceComponent->ApplyGameplayEffectToTarget(BuffEffect, SourceComponent, 1.f, Qualifier);
+		BuffHandle = SourceComponent->ApplyGameplayEffectToTarget(BuffEffect, SourceComponent, 1.f);
 		Test->TestTrue(SKILL_TEST_TEXT("Number of Source GameplayEffect: %d == 2", SourceComponent->GetNumActiveGameplayEffect()), SourceComponent->GetNumActiveGameplayEffect() == 2);
 
 		// Check that the buff was buffed
@@ -1380,11 +1304,8 @@ bool GameplayEffectsTest_DamageBuffBuff_FullSnapshot(UWorld *World, FAutomationT
 		BaseDmgEffect->Modifiers[0].OwnedTags.AddTag(IGameplayTagsModule::RequestGameplayTag(FName(TEXT("Damage.Basic"))));
 		BaseDmgEffect->Duration.Value = UGameplayEffect::INFINITE_DURATION;
 
-		FModifierQualifier Qualifier = FModifierQualifier();
-		Qualifier.PredictionKey(TEST_PREDICTIONKEY);
-
 		// Apply to target
-		DamageHandle = SourceComponent->ApplyGameplayEffectToTarget(BaseDmgEffect, DestComponent, 1.f, Qualifier);
+		DamageHandle = SourceComponent->ApplyGameplayEffectToTarget(BaseDmgEffect, DestComponent, 1.f);
 
 		float ExpectedValue = (StartHealth + (DamageValue * (DamageBuffMultiplier + DamageBuffMultiplierBonus)));
 		float ActualValue = DestComponent->GetSet<USkillSystemTestAttributeSet>()->Health;
@@ -1480,11 +1401,8 @@ bool GameplayEffectsTest_DamageBuffBuff_SnapshotLink(UWorld *World, FAutomationT
 		BuffEffect->Duration.Value = UGameplayEffect::INFINITE_DURATION;
 		BuffEffect->CopyPolicy = EGameplayEffectCopyPolicy::AlwaysSnapshot;
 
-		FModifierQualifier Qualifier = FModifierQualifier();
-		Qualifier.PredictionKey(TEST_PREDICTIONKEY);
-
 		// Apply to target
-		BuffBuffHandle = SourceComponent->ApplyGameplayEffectToTarget(BuffEffect, SourceComponent, 1.f, Qualifier);
+		BuffBuffHandle = SourceComponent->ApplyGameplayEffectToTarget(BuffEffect, SourceComponent, 1.f);
 
 		Test->TestTrue(SKILL_TEST_TEXT("Number of Source GameplayEffect: %d == 1", SourceComponent->GetNumActiveGameplayEffect()), SourceComponent->GetNumActiveGameplayEffect() == 1);
 	}
@@ -1501,11 +1419,8 @@ bool GameplayEffectsTest_DamageBuffBuff_SnapshotLink(UWorld *World, FAutomationT
 		BuffEffect->Duration.Value = UGameplayEffect::INFINITE_DURATION;
 		BuffEffect->CopyPolicy = EGameplayEffectCopyPolicy::AlwaysLink;
 
-		FModifierQualifier Qualifier = FModifierQualifier();
-		Qualifier.PredictionKey(TEST_PREDICTIONKEY);
-
 		// Apply to target
-		BuffHandle = SourceComponent->ApplyGameplayEffectToTarget(BuffEffect, SourceComponent, 1.f, Qualifier);
+		BuffHandle = SourceComponent->ApplyGameplayEffectToTarget(BuffEffect, SourceComponent, 1.f);
 		Test->TestTrue(SKILL_TEST_TEXT("Number of Source GameplayEffect: %d == 2", SourceComponent->GetNumActiveGameplayEffect()), SourceComponent->GetNumActiveGameplayEffect() == 2);
 
 		// Check that the buff was buffed
@@ -1527,11 +1442,8 @@ bool GameplayEffectsTest_DamageBuffBuff_SnapshotLink(UWorld *World, FAutomationT
 		BaseDmgEffect->Modifiers[0].OwnedTags.AddTag(IGameplayTagsModule::RequestGameplayTag(FName(TEXT("Damage.Basic"))));
 		BaseDmgEffect->Duration.Value = UGameplayEffect::INFINITE_DURATION;
 
-		FModifierQualifier Qualifier = FModifierQualifier();
-		Qualifier.PredictionKey(TEST_PREDICTIONKEY);
-
 		// Apply to target
-		DamageHandle = SourceComponent->ApplyGameplayEffectToTarget(BaseDmgEffect, DestComponent, 1.f, Qualifier);
+		DamageHandle = SourceComponent->ApplyGameplayEffectToTarget(BaseDmgEffect, DestComponent, 1.f);
 
 		float ExpectedValue = (StartHealth + (DamageValue * (DamageBuffMultiplier + DamageBuffMultiplierBonus)));
 		float ActualValue = DestComponent->GetSet<USkillSystemTestAttributeSet>()->Health;
@@ -1626,11 +1538,8 @@ bool GameplayEffectsTest_DamageAppliesBuff(UWorld *World, FAutomationTestBase *
 		BaseDmgEffect->Duration.Value = UGameplayEffect::INSTANT_APPLICATION;
 		BaseDmgEffect->TargetEffects.Add(BuffEffect);
 
-		FModifierQualifier Qualifier = FModifierQualifier();
-		Qualifier.PredictionKey(TEST_PREDICTIONKEY);
-
 		// Apply to target
-		SourceComponent->ApplyGameplayEffectToTarget(BaseDmgEffect, DestComponent, 1.f, Qualifier);
+		SourceComponent->ApplyGameplayEffectToTarget(BaseDmgEffect, DestComponent, 1.f);
 
 		float ExpectedValue = (StartHealth - DamageValue);
 		float ActualValue = DestComponent->GetSet<USkillSystemTestAttributeSet>()->Health;
@@ -1704,11 +1613,8 @@ bool GameplayEffectsTest_BuffAppliesBuff(UWorld *World, FAutomationTestBase * Te
 		DummyBuffEffect->Modifiers[0].TargetEffect = BuffEffect;
 		DummyBuffEffect->Duration.SetValue(UGameplayEffect::INFINITE_DURATION);
 
-		FModifierQualifier Qualifier = FModifierQualifier();
-		Qualifier.PredictionKey(TEST_PREDICTIONKEY);
-
 		// Apply to target
-		SourceComponent->ApplyGameplayEffectToTarget(DummyBuffEffect, SourceComponent, 1.f, Qualifier);
+		SourceComponent->ApplyGameplayEffectToTarget(DummyBuffEffect, SourceComponent, 1.f);
 	}
 
 	// apply damage to source to make sure it didn't get the protection buff
@@ -1743,11 +1649,8 @@ bool GameplayEffectsTest_BuffAppliesBuff(UWorld *World, FAutomationTestBase * Te
 		BaseDmgEffect->Modifiers[0].OwnedTags.AddTag(IGameplayTagsModule::RequestGameplayTag(FName(TEXT("Damage.Buffable"))));
 		BaseDmgEffect->Duration.Value = UGameplayEffect::INSTANT_APPLICATION;
 
-		FModifierQualifier Qualifier = FModifierQualifier();
-		Qualifier.PredictionKey(TEST_PREDICTIONKEY);
-
 		// Apply to target
-		SourceComponent->ApplyGameplayEffectToTarget(BaseDmgEffect, DestComponent, 1.f, Qualifier);
+		SourceComponent->ApplyGameplayEffectToTarget(BaseDmgEffect, DestComponent, 1.f);
 
 		float ExpectedValue = StartHealth - DamageValue;
 		float ActualValue = DestComponent->GetSet<USkillSystemTestAttributeSet>()->Health;
@@ -1767,11 +1670,8 @@ bool GameplayEffectsTest_BuffAppliesBuff(UWorld *World, FAutomationTestBase * Te
 		BaseDmgEffect->Modifiers[0].OwnedTags.AddTag(IGameplayTagsModule::RequestGameplayTag(FName(TEXT("Damage.Basic"))));
 		BaseDmgEffect->Duration.Value = UGameplayEffect::INSTANT_APPLICATION;
 
-		FModifierQualifier Qualifier = FModifierQualifier();
-		Qualifier.PredictionKey(TEST_PREDICTIONKEY);
-
 		// Apply to target
-		SourceComponent->ApplyGameplayEffectToTarget(BaseDmgEffect, DestComponent, 1.f, Qualifier);
+		SourceComponent->ApplyGameplayEffectToTarget(BaseDmgEffect, DestComponent, 1.f);
 
 		float ExpectedValue = (StartHealth - DamageValue - (DamageValue / DamageProtectionDivisor));
 		float ActualValue = DestComponent->GetSet<USkillSystemTestAttributeSet>()->Health;
@@ -1791,11 +1691,8 @@ bool GameplayEffectsTest_BuffAppliesBuff(UWorld *World, FAutomationTestBase * Te
 		BaseDmgEffect->Modifiers[0].OwnedTags.AddTag(IGameplayTagsModule::RequestGameplayTag(FName(TEXT("Damage.Basic"))));
 		BaseDmgEffect->Duration.Value = UGameplayEffect::INSTANT_APPLICATION;
 
-		FModifierQualifier Qualifier = FModifierQualifier();
-		Qualifier.PredictionKey(TEST_PREDICTIONKEY);
-
 		// Apply to target
-		SourceComponent->ApplyGameplayEffectToTarget(BaseDmgEffect, DestComponent, 1.f, Qualifier);
+		SourceComponent->ApplyGameplayEffectToTarget(BaseDmgEffect, DestComponent, 1.f);
 
 		float ExpectedValue = (StartHealth - DamageValue - (2 * DamageValue / DamageProtectionDivisor));
 		float ActualValue = DestComponent->GetSet<USkillSystemTestAttributeSet>()->Health;
@@ -1854,11 +1751,8 @@ bool GameplayEffectsTest_BuffIndirection(UWorld *World, FAutomationTestBase * Te
 		DummyBuffEffect2->Modifiers[0].TargetEffect = DummyBuffEffect;
 		DummyBuffEffect2->Duration.SetValue(UGameplayEffect::INFINITE_DURATION);
 
-		FModifierQualifier Qualifier = FModifierQualifier();
-		Qualifier.PredictionKey(TEST_PREDICTIONKEY);
-
 		// Apply to target
-		SourceComponent->ApplyGameplayEffectToTarget(DummyBuffEffect2, SourceComponent, 1.f, Qualifier);
+		SourceComponent->ApplyGameplayEffectToTarget(DummyBuffEffect2, SourceComponent, 1.f);
 	}
 
 	// Apply Damage
@@ -1872,11 +1766,8 @@ bool GameplayEffectsTest_BuffIndirection(UWorld *World, FAutomationTestBase * Te
 		BaseDmgEffect->Modifiers[0].OwnedTags.AddTag(IGameplayTagsModule::RequestGameplayTag(FName(TEXT("Damage.Basic"))));
 		BaseDmgEffect->Duration.Value = UGameplayEffect::INSTANT_APPLICATION;
 
-		FModifierQualifier Qualifier = FModifierQualifier();
-		Qualifier.PredictionKey(TEST_PREDICTIONKEY);
-
 		// Apply to target
-		SourceComponent->ApplyGameplayEffectToTarget(BaseDmgEffect, DestComponent, 1.f, Qualifier);
+		SourceComponent->ApplyGameplayEffectToTarget(BaseDmgEffect, DestComponent, 1.f);
 
 		float ExpectedValue = StartHealth - DamageValue;
 		float ActualValue = DestComponent->GetSet<USkillSystemTestAttributeSet>()->Health;
@@ -1896,11 +1787,8 @@ bool GameplayEffectsTest_BuffIndirection(UWorld *World, FAutomationTestBase * Te
 		BaseDmgEffect->Modifiers[0].OwnedTags.AddTag(IGameplayTagsModule::RequestGameplayTag(FName(TEXT("Damage.Buffable"))));
 		BaseDmgEffect->Duration.Value = UGameplayEffect::INSTANT_APPLICATION;
 
-		FModifierQualifier Qualifier = FModifierQualifier();
-		Qualifier.PredictionKey(TEST_PREDICTIONKEY);
-
 		// Apply to target
-		DestComponent->ApplyGameplayEffectToTarget(BaseDmgEffect, SourceComponent, 1.f, Qualifier);
+		DestComponent->ApplyGameplayEffectToTarget(BaseDmgEffect, SourceComponent, 1.f);
 
 		float ExpectedValue = StartHealth - DamageValue;
 		float ActualValue = SourceComponent->GetSet<USkillSystemTestAttributeSet>()->Health;
@@ -1920,11 +1808,8 @@ bool GameplayEffectsTest_BuffIndirection(UWorld *World, FAutomationTestBase * Te
 		BaseDmgEffect->Modifiers[0].OwnedTags.AddTag(IGameplayTagsModule::RequestGameplayTag(FName(TEXT("Damage.Basic"))));
 		BaseDmgEffect->Duration.Value = UGameplayEffect::INSTANT_APPLICATION;
 
-		FModifierQualifier Qualifier = FModifierQualifier();
-		Qualifier.PredictionKey(TEST_PREDICTIONKEY);
-
 		// Apply to target
-		DestComponent->ApplyGameplayEffectToTarget(BaseDmgEffect, SourceComponent, 1.f, Qualifier);
+		DestComponent->ApplyGameplayEffectToTarget(BaseDmgEffect, SourceComponent, 1.f);
 
 		float ExpectedValue = StartHealth - DamageValue - (DamageValue / DamageProtectionDivisor);
 		float ActualValue = SourceComponent->GetSet<USkillSystemTestAttributeSet>()->Health;
@@ -1970,10 +1855,7 @@ bool GameplayEffectsTest_DurationDamage(UWorld *World, FAutomationTestBase * Tes
 		BaseDmgEffect->Duration.Value = Duration;
 		BaseDmgEffect->Period.Value = UGameplayEffect::NO_PERIOD;
 
-		FModifierQualifier Qualifier = FModifierQualifier();
-		Qualifier.PredictionKey(TEST_PREDICTIONKEY);
-
-		SourceComponent->ApplyGameplayEffectToTarget(BaseDmgEffect, DestComponent, 1.f, Qualifier);
+		SourceComponent->ApplyGameplayEffectToTarget(BaseDmgEffect, DestComponent, 1.f);
 
 		// The effect should instantly apply without ticking (for now at least)
 		float ActualValue = DestComponent->GetSet<USkillSystemTestAttributeSet>()->Health;
@@ -2126,10 +2008,7 @@ bool GameplayEffectsTest_LifestealExtension(UWorld *World, FAutomationTestBase *
 		LifestealEffect->Duration.Value = UGameplayEffect::INFINITE_DURATION;
 		LifestealEffect->Period.Value = UGameplayEffect::NO_PERIOD;
 
-		FModifierQualifier Qualifier = FModifierQualifier();
-		Qualifier.PredictionKey(TEST_PREDICTIONKEY);
-
-		SourceComponent->ApplyGameplayEffectToSelf(LifestealEffect, 1.f, SourceActor, Qualifier);
+		SourceComponent->ApplyGameplayEffectToSelf(LifestealEffect, 1.f, SourceActor);
 	}
 	
 	{
@@ -2205,10 +2084,7 @@ bool GameplayEffectsTest_ShieldExtension(UWorld *World, FAutomationTestBase * Te
 		ShieldEffect->Duration.Value = UGameplayEffect::INFINITE_DURATION;
 		ShieldEffect->Period.Value = UGameplayEffect::NO_PERIOD;
 
-		FModifierQualifier Qualifier = FModifierQualifier();
-		Qualifier.PredictionKey(TEST_PREDICTIONKEY);
-
-		AppliedHandle = DestComponent->ApplyGameplayEffectToSelf(ShieldEffect, 1.f, DestActor, Qualifier);
+		AppliedHandle = DestComponent->ApplyGameplayEffectToSelf(ShieldEffect, 1.f, DestActor);
 	}
 
 	{
@@ -2324,11 +2200,8 @@ bool GameplayEffectsTest_ShieldExtensionMultiple(UWorld *World, FAutomationTestB
 		ShieldEffect->Duration.Value = UGameplayEffect::INFINITE_DURATION;
 		ShieldEffect->Period.Value = UGameplayEffect::NO_PERIOD;
 
-		FModifierQualifier Qualifier = FModifierQualifier();
-		Qualifier.PredictionKey(TEST_PREDICTIONKEY);
-
-		AppliedHandle_1 = DestComponent->ApplyGameplayEffectToSelf(ShieldEffect, 1.f, DestActor, Qualifier);
-		AppliedHandle_2 = DestComponent->ApplyGameplayEffectToSelf(ShieldEffect, 1.f, DestActor, Qualifier);
+		AppliedHandle_1 = DestComponent->ApplyGameplayEffectToSelf(ShieldEffect, 1.f, DestActor);
+		AppliedHandle_2 = DestComponent->ApplyGameplayEffectToSelf(ShieldEffect, 1.f, DestActor);
 	}
 
 	{
@@ -2688,11 +2561,8 @@ bool GameplayEffectsTest_InstantDamageRequiredTag(UWorld *World, FAutomationTest
 		BaseProtectEffect->CopyPolicy = EGameplayEffectCopyPolicy::AlwaysLink;
 		BaseProtectEffect->GameplayEffectRequiredTags.AddTag(IGameplayTagsModule::RequestGameplayTag(FName(TEXT("Damage.Type2"))));
 
-		FModifierQualifier Qualifier = FModifierQualifier();
-		Qualifier.PredictionKey(TEST_PREDICTIONKEY);
-
 		// Apply to self
-		DestComponent->ApplyGameplayEffectToTarget(BaseProtectEffect, DestComponent, 1.f, Qualifier);
+		DestComponent->ApplyGameplayEffectToTarget(BaseProtectEffect, DestComponent, 1.f);
 	}
 
 	// Apply Damage
@@ -2781,11 +2651,8 @@ bool GameplayEffectsTest_InstantDamageIgnoreTag(UWorld *World, FAutomationTestBa
 		BaseProtectEffect->CopyPolicy = EGameplayEffectCopyPolicy::AlwaysLink;
 		BaseProtectEffect->GameplayEffectIgnoreTags.AddTag(IGameplayTagsModule::RequestGameplayTag(FName(TEXT("Damage.Type1"))));
 
-		FModifierQualifier Qualifier = FModifierQualifier();
-		Qualifier.PredictionKey(TEST_PREDICTIONKEY);
-
 		// Apply to self
-		DestComponent->ApplyGameplayEffectToTarget(BaseProtectEffect, DestComponent, 1.f, Qualifier);
+		DestComponent->ApplyGameplayEffectToTarget(BaseProtectEffect, DestComponent, 1.f);
 	}
 
 	// Apply Damage
@@ -2873,11 +2740,8 @@ bool GameplayEffectsTest_InstantDamageModifierPassesTag(UWorld *World, FAutomati
 		BaseDmgEffect->Modifiers[0].OwnedTags.AddTag(IGameplayTagsModule::RequestGameplayTag(FName(TEXT("Damage.Type1"))));		// When I am applied, the damage modifier gets this tag.
 		BaseDmgEffect->Duration.SetValue(UGameplayEffect::INFINITE_DURATION);
 
-		FModifierQualifier Qualifier = FModifierQualifier();
-		Qualifier.PredictionKey(TEST_PREDICTIONKEY);
-
 		// Apply to self
-		SourceComponent->ApplyGameplayEffectToTarget(BaseDmgEffect, SourceComponent, 1.f, Qualifier);
+		SourceComponent->ApplyGameplayEffectToTarget(BaseDmgEffect, SourceComponent, 1.f);
 	}
 
 	// Setup a GE to modify IncomingGEs
@@ -2895,11 +2759,8 @@ bool GameplayEffectsTest_InstantDamageModifierPassesTag(UWorld *World, FAutomati
 		BaseProtectEffect->Duration.SetValue(UGameplayEffect::INFINITE_DURATION);
 		BaseProtectEffect->CopyPolicy = EGameplayEffectCopyPolicy::AlwaysLink;
 
-		FModifierQualifier Qualifier = FModifierQualifier();
-		Qualifier.PredictionKey(TEST_PREDICTIONKEY);
-
 		// Apply to self
-		DestComponent->ApplyGameplayEffectToTarget(BaseProtectEffect, DestComponent, 1.f, Qualifier);
+		DestComponent->ApplyGameplayEffectToTarget(BaseProtectEffect, DestComponent, 1.f);
 	}
 
 	// Apply Damage
@@ -2961,11 +2822,8 @@ bool GameplayEffectsTest_InstantDamageModifierTag(UWorld *World, FAutomationTest
 		BaseDmgEffect->Modifiers[0].OwnedTags.AddTag(IGameplayTagsModule::RequestGameplayTag(FName(TEXT("Damage.Type1"))));
 		BaseDmgEffect->Duration.SetValue(UGameplayEffect::INFINITE_DURATION);
 
-		FModifierQualifier Qualifier = FModifierQualifier();
-		Qualifier.PredictionKey(TEST_PREDICTIONKEY);
-
 		// Apply to self
-		SourceComponent->ApplyGameplayEffectToTarget(BaseDmgEffect, SourceComponent, 1.f, Qualifier);
+		SourceComponent->ApplyGameplayEffectToTarget(BaseDmgEffect, SourceComponent, 1.f);
 	}
 
 	// Setup a GE to modify IncomingGEs
@@ -2983,11 +2841,8 @@ bool GameplayEffectsTest_InstantDamageModifierTag(UWorld *World, FAutomationTest
 		BaseProtectEffect->Duration.SetValue(UGameplayEffect::INFINITE_DURATION);
 		BaseProtectEffect->CopyPolicy = EGameplayEffectCopyPolicy::AlwaysLink;
 
-		FModifierQualifier Qualifier = FModifierQualifier();
-		Qualifier.PredictionKey(TEST_PREDICTIONKEY);
-
 		// Apply to self
-		DestComponent->ApplyGameplayEffectToTarget(BaseProtectEffect, DestComponent, 1.f, Qualifier);
+		DestComponent->ApplyGameplayEffectToTarget(BaseProtectEffect, DestComponent, 1.f);
 	}
 
 	// Apply Damage
@@ -3175,10 +3030,7 @@ bool GameplayEffectsTest_DotDamage_ScalingProperty_Snapshot(UWorld *World, FAuto
 		BaseDmgEffect->LevelInfo.InheritLevelFromOwner = false;
 		BaseDmgEffect->LevelInfo.TakeSnapshotOnInit = true;								// But just a snapshot of their SpellDamage when we are applied
 
-		FModifierQualifier Qualifier = FModifierQualifier();
-		Qualifier.PredictionKey(TEST_PREDICTIONKEY);
-
-		SourceComponent->ApplyGameplayEffectToTarget(BaseDmgEffect, DestComponent, FGameplayEffectLevelSpec::INVALID_LEVEL, Qualifier);
+		SourceComponent->ApplyGameplayEffectToTarget(BaseDmgEffect, DestComponent);
 	}
 
 
@@ -3264,10 +3116,7 @@ bool GameplayEffectsTest_DotDamage_ScalingProperty_Dynamic(UWorld *World, FAutom
 		BaseDmgEffect->LevelInfo.InheritLevelFromOwner = false;
 		BaseDmgEffect->LevelInfo.TakeSnapshotOnInit = false;							// But level is dynamic, if SpellDamage changers after we apply, we update.
 
-		FModifierQualifier Qualifier = FModifierQualifier();
-		Qualifier.PredictionKey(TEST_PREDICTIONKEY);
-
-		SourceComponent->ApplyGameplayEffectToTarget(BaseDmgEffect, DestComponent, FGameplayEffectLevelSpec::INVALID_LEVEL, Qualifier);
+		SourceComponent->ApplyGameplayEffectToTarget(BaseDmgEffect, DestComponent);
 		GameplayTest_TickWorld(World, 0.0001f); // Move our Effects from the pending stack to the active stack inside the timer manager, this starts the clock for effect execution
 
 		float SpellDamageTest = SourceComponent->GetSet<USkillSystemTestAttributeSet>()->SpellDamage;
@@ -3363,10 +3212,7 @@ bool GameplayEffectsTest_MetaAttributes(UWorld *World, FAutomationTestBase * Tes
 		StrengthMaxHealhEffect->LevelInfo.InheritLevelFromOwner = false;
 		StrengthMaxHealhEffect->LevelInfo.TakeSnapshotOnInit = false;						// But level is dynamic, if SpellDamage changers after we apply, we update.
 
-		FModifierQualifier Qualifier = FModifierQualifier();
-		Qualifier.PredictionKey(TEST_PREDICTIONKEY);
-
-		SourceComponent->ApplyGameplayEffectToTarget(StrengthMaxHealhEffect, SourceComponent, FGameplayEffectLevelSpec::INVALID_LEVEL, Qualifier);
+		SourceComponent->ApplyGameplayEffectToTarget(StrengthMaxHealhEffect, SourceComponent);
 
 		// Strength starts at 0, so confirm it did nothing yet.
 		float ActualValue = SourceComponent->GetSet<USkillSystemTestAttributeSet>()->MaxHealth;
@@ -3441,11 +3287,8 @@ bool GameplayEffectsTest_TagOrdering(UWorld *World, FAutomationTestBase * Test)
 		BaseDmgEffect->GameplayEffectTags.AddTag(IGameplayTagsModule::RequestGameplayTag(FName(TEXT("Buff"))));
 		BaseDmgEffect->GameplayEffectRequiredTags.AddTag(IGameplayTagsModule::RequestGameplayTag(FName(TEXT("Damage"))));
 
-		FModifierQualifier Qualifier = FModifierQualifier();
-		Qualifier.PredictionKey(TEST_PREDICTIONKEY);
-
 		// Apply to self
-		SourceComponent->ApplyGameplayEffectToTarget(BaseDmgEffect, SourceComponent, 1.f, Qualifier);
+		SourceComponent->ApplyGameplayEffectToTarget(BaseDmgEffect, SourceComponent, 1.f);
 	}
 
 	{
@@ -3464,11 +3307,8 @@ bool GameplayEffectsTest_TagOrdering(UWorld *World, FAutomationTestBase * Test)
 		BaseDmgEffect->GameplayEffectTags.AddTag(IGameplayTagsModule::RequestGameplayTag(FName(TEXT("Buff"))));
 		BaseDmgEffect->GameplayEffectRequiredTags.AddTag(IGameplayTagsModule::RequestGameplayTag(FName(TEXT("Damage"))));
 
-		FModifierQualifier Qualifier = FModifierQualifier();
-		Qualifier.PredictionKey(TEST_PREDICTIONKEY);
-
 		// Apply to self
-		SourceComponent->ApplyGameplayEffectToTarget(BaseDmgEffect, SourceComponent, 1.f, Qualifier);
+		SourceComponent->ApplyGameplayEffectToTarget(BaseDmgEffect, SourceComponent, 1.f);
 	}
 
 	// Apply Damage
@@ -4525,11 +4365,8 @@ bool GameplayEffectsTest_StackingRemovingModifiers(UWorld *World, FAutomationTes
 		BaseModEffect->CopyPolicy = EGameplayEffectCopyPolicy::AlwaysLink;
 		BaseModEffect->GameplayEffectRequiredTags.AddTag(IGameplayTagsModule::RequestGameplayTag(FName(TEXT("Stack"))));
 
-		FModifierQualifier Qualifier = FModifierQualifier();
-		Qualifier.PredictionKey(TEST_PREDICTIONKEY);
-
 		// Apply to self
-		DestComponent->ApplyGameplayEffectToTarget(BaseModEffect, DestComponent, 1.f, Qualifier);
+		DestComponent->ApplyGameplayEffectToTarget(BaseModEffect, DestComponent, 1.f);
 	}
 
 	{
@@ -4669,11 +4506,8 @@ bool GameplayEffectsTest_StackingAddingModifiers(UWorld *World, FAutomationTestB
 		BaseModEffect->CopyPolicy = EGameplayEffectCopyPolicy::AlwaysLink;
 		BaseModEffect->GameplayEffectRequiredTags.AddTag(IGameplayTagsModule::RequestGameplayTag(FName(TEXT("Stack"))));
 
-		FModifierQualifier Qualifier = FModifierQualifier();
-		Qualifier.PredictionKey(TEST_PREDICTIONKEY);
-
 		// Apply to self
-		DestComponent->ApplyGameplayEffectToTarget(BaseModEffect, DestComponent, 1.f, Qualifier);
+		DestComponent->ApplyGameplayEffectToTarget(BaseModEffect, DestComponent, 1.f);
 
 		GameplayTest_TickWorld(World, 0.0001f); // Move our Effects from the pending stack to the active stack inside the timer manager, this starts the clock for effect execution
 		GameplayTest_TickWorld(World, 0.1f); // Offset the current time from the start of the period so that floating point issues don't affect the test
@@ -4723,11 +4557,8 @@ bool GameplayEffectsTest_ImmunityIncoming(UWorld *World, FAutomationTestBase * T
 		BaseImmunityEffect->GameplayEffectRequiredTags.AddTag(IGameplayTagsModule::RequestGameplayTag(FName(TEXT("Immune"))));
 		BaseImmunityEffect->AppliesImmunityTo = EGameplayImmunity::IncomingGE;
 
-		FModifierQualifier Qualifier = FModifierQualifier();
-		Qualifier.PredictionKey(TEST_PREDICTIONKEY);
-
 		// Apply to self
-		DestComponent->ApplyGameplayEffectToTarget(BaseImmunityEffect, DestComponent, 1.f, Qualifier);
+		DestComponent->ApplyGameplayEffectToTarget(BaseImmunityEffect, DestComponent, 1.f);
 	}
 
 	// Apply Damage
@@ -4784,11 +4615,8 @@ bool GameplayEffectsTest_ImmunityOutgoing(UWorld *World, FAutomationTestBase * T
 		BaseImmunityEffect->GameplayEffectRequiredTags.AddTag(IGameplayTagsModule::RequestGameplayTag(FName(TEXT("Immune"))));
 		BaseImmunityEffect->AppliesImmunityTo = EGameplayImmunity::OutgoingGE;
 
-		FModifierQualifier Qualifier = FModifierQualifier();
-		Qualifier.PredictionKey(TEST_PREDICTIONKEY);
-
 		// Apply to self
-		SourceComponent->ApplyGameplayEffectToTarget(BaseImmunityEffect, SourceComponent, 1.f, Qualifier);
+		SourceComponent->ApplyGameplayEffectToTarget(BaseImmunityEffect, SourceComponent, 1.f);
 	}
 
 	// Apply Damage
@@ -4848,11 +4676,8 @@ bool GameplayEffectsTest_ImmunityMod(UWorld *World, FAutomationTestBase * Test)
 		BaseImmunityEffect->GameplayEffectRequiredTags.AddTag(IGameplayTagsModule::RequestGameplayTag(FName(TEXT("Damage.Buff"))));
 		BaseImmunityEffect->AppliesImmunityTo = EGameplayImmunity::IncomingGE;
 
-		FModifierQualifier Qualifier = FModifierQualifier();
-		Qualifier.PredictionKey(TEST_PREDICTIONKEY);
-
 		// Apply to self
-		DestComponent->ApplyGameplayEffectToTarget(BaseImmunityEffect, DestComponent, 1.f, Qualifier);
+		DestComponent->ApplyGameplayEffectToTarget(BaseImmunityEffect, DestComponent, 1.f);
 	}
 
 	// Setup a GE to modify OutgoingGEs
@@ -4868,11 +4693,8 @@ bool GameplayEffectsTest_ImmunityMod(UWorld *World, FAutomationTestBase * Test)
 		BaseDmgEffect->Modifiers[0].OwnedTags.AddTag(IGameplayTagsModule::RequestGameplayTag(FName(TEXT("Damage.Buff"))));
 		BaseDmgEffect->Duration.SetValue(UGameplayEffect::INFINITE_DURATION);
 
-		FModifierQualifier Qualifier = FModifierQualifier();
-		Qualifier.PredictionKey(TEST_PREDICTIONKEY);
-
 		// Apply to self
-		SourceComponent->ApplyGameplayEffectToTarget(BaseDmgEffect, SourceComponent, 1.f, Qualifier);
+		SourceComponent->ApplyGameplayEffectToTarget(BaseDmgEffect, SourceComponent, 1.f);
 	}
 
 	// Apply Damage
@@ -4930,11 +4752,8 @@ bool GameplayEffectsTest_ImmunityBlockedBuff(UWorld *World, FAutomationTestBase
 		BaseImmunityEffect->GameplayEffectRequiredTags.AddTag(IGameplayTagsModule::RequestGameplayTag(FName(TEXT("Damage.Buff"))));
 		BaseImmunityEffect->AppliesImmunityTo = EGameplayImmunity::IncomingGE;
 
-		FModifierQualifier Qualifier = FModifierQualifier();
-		Qualifier.PredictionKey(TEST_PREDICTIONKEY);
-
 		// Apply to self
-		DestComponent->ApplyGameplayEffectToTarget(BaseImmunityEffect, DestComponent, 1.f, Qualifier);
+		DestComponent->ApplyGameplayEffectToTarget(BaseImmunityEffect, DestComponent, 1.f);
 	}
 
 	// Apply base damage and a buff that causes extra damage
@@ -5002,11 +4821,8 @@ bool GameplayEffectsTest_ImmunityBlockedBaseAndBuff(UWorld *World, FAutomationTe
 		BaseImmunityEffect->GameplayEffectRequiredTags.AddTag(IGameplayTagsModule::RequestGameplayTag(FName(TEXT("Damage.Basic"))));
 		BaseImmunityEffect->AppliesImmunityTo = EGameplayImmunity::IncomingGE;
 
-		FModifierQualifier Qualifier = FModifierQualifier();
-		Qualifier.PredictionKey(TEST_PREDICTIONKEY);
-
 		// Apply to self
-		DestComponent->ApplyGameplayEffectToTarget(BaseImmunityEffect, DestComponent, 1.f, Qualifier);
+		DestComponent->ApplyGameplayEffectToTarget(BaseImmunityEffect, DestComponent, 1.f);
 	}
 
 	// Apply base damage and a buff that causes extra damage
@@ -5076,11 +4892,8 @@ bool GameplayEffectsTest_ImmunityActiveGE(UWorld *World, FAutomationTestBase * T
 		BuffEffect->GameplayEffectTags.AddTag(IGameplayTagsModule::RequestGameplayTag(FName(TEXT("Damage.Reduce"))));
 		BuffEffect->Duration.SetValue(UGameplayEffect::INFINITE_DURATION);
 
-		FModifierQualifier Qualifier = FModifierQualifier();
-		Qualifier.PredictionKey(TEST_PREDICTIONKEY);
-
 		// Apply to self
-		DestComponent->ApplyGameplayEffectToTarget(BuffEffect, DestComponent, 1.f, Qualifier);
+		DestComponent->ApplyGameplayEffectToTarget(BuffEffect, DestComponent, 1.f);
 	}
 
 	// Apply Damage to verify the buff is working
@@ -5176,11 +4989,8 @@ bool GameplayEffectsTest_ImmunityActiveMod(UWorld *World, FAutomationTestBase *
 		BuffEffect->Modifiers[0].OwnedTags.AddTag(IGameplayTagsModule::RequestGameplayTag(FName(TEXT("Damage.Reduce"))));
 		BuffEffect->Duration.SetValue(UGameplayEffect::INFINITE_DURATION);
 
-		FModifierQualifier Qualifier = FModifierQualifier();
-		Qualifier.PredictionKey(TEST_PREDICTIONKEY);
-
 		// Apply to self
-		DestComponent->ApplyGameplayEffectToTarget(BuffEffect, DestComponent, 1.f, Qualifier);
+		DestComponent->ApplyGameplayEffectToTarget(BuffEffect, DestComponent, 1.f);
 	}
 
 	// Apply Damage to verify the buff is working
@@ -5348,11 +5158,8 @@ bool GameplayEffectsTest_ChanceToExecuteOnActiveGEMod(UWorld *World, FAutomation
 		BuffDmgEffect->Modifiers[0].RequiredTags.AddTag(IGameplayTagsModule::RequestGameplayTag(FName(TEXT("Damage.Basic"))));
 		BuffDmgEffect->Duration.SetValue(UGameplayEffect::INFINITE_DURATION);
 
-		FModifierQualifier Qualifier = FModifierQualifier();
-		Qualifier.PredictionKey(TEST_PREDICTIONKEY);
-
 		// Apply to self
-		SourceComponent->ApplyGameplayEffectToTarget(BuffDmgEffect, SourceComponent, 1.f, Qualifier);
+		SourceComponent->ApplyGameplayEffectToTarget(BuffDmgEffect, SourceComponent, 1.f);
 	}
 
 	// attempt to modify the buff but fail because of a 0.f chance to apply
@@ -5468,11 +5275,8 @@ bool GameplayEffectsTest_ChanceToExecuteOnActiveGEImmunity(UWorld *World, FAutom
 		BuffDmgEffect->Modifiers[0].RequiredTags.AddTag(IGameplayTagsModule::RequestGameplayTag(FName(TEXT("Damage.Basic"))));
 		BuffDmgEffect->Duration.SetValue(UGameplayEffect::INFINITE_DURATION);
 
-		FModifierQualifier Qualifier = FModifierQualifier();
-		Qualifier.PredictionKey(TEST_PREDICTIONKEY);
-
 		// Apply to self
-		SourceComponent->ApplyGameplayEffectToTarget(BuffDmgEffect, SourceComponent, 1.f, Qualifier);
+		SourceComponent->ApplyGameplayEffectToTarget(BuffDmgEffect, SourceComponent, 1.f);
 	}
 
 	// attempt to remove the buff but fail because of a 0.f chance to apply
@@ -5583,11 +5387,8 @@ bool GameplayEffectsTest_ChanceToExecuteOnOutgoingGEMod(UWorld *World, FAutomati
 		BuffDmgEffect->Duration.SetValue(UGameplayEffect::INFINITE_DURATION);
 		BuffDmgEffect->ChanceToExecuteOnGameplayEffect.SetValue(0.f);
 
-		FModifierQualifier Qualifier = FModifierQualifier();
-		Qualifier.PredictionKey(TEST_PREDICTIONKEY);
-
 		// Apply to self
-		SourceComponent->ApplyGameplayEffectToTarget(BuffDmgEffect, SourceComponent, 1.f, Qualifier);
+		SourceComponent->ApplyGameplayEffectToTarget(BuffDmgEffect, SourceComponent, 1.f);
 	}
 
 	{
@@ -5627,11 +5428,8 @@ bool GameplayEffectsTest_ChanceToExecuteOnOutgoingGEMod(UWorld *World, FAutomati
 		BuffDmgEffect->Duration.SetValue(UGameplayEffect::INFINITE_DURATION);
 		BuffDmgEffect->ChanceToExecuteOnGameplayEffect.SetValue(1.f);
 
-		FModifierQualifier Qualifier = FModifierQualifier();
-		Qualifier.PredictionKey(TEST_PREDICTIONKEY);
-
 		// Apply to self
-		SourceComponent->ApplyGameplayEffectToTarget(BuffDmgEffect, SourceComponent, 1.f, Qualifier);
+		SourceComponent->ApplyGameplayEffectToTarget(BuffDmgEffect, SourceComponent, 1.f);
 	}
 
 	{
@@ -5687,11 +5485,8 @@ bool GameplayEffectsTest_ChanceToExecuteOnOutgoingGEImmunity(UWorld *World, FAut
 		ModBuffEffect->Duration.SetValue(UGameplayEffect::INFINITE_DURATION);
 		ModBuffEffect->ChanceToExecuteOnGameplayEffect.SetValue(0.f);
 
-		FModifierQualifier Qualifier = FModifierQualifier();
-		Qualifier.PredictionKey(TEST_PREDICTIONKEY);
-
 		// Apply to self
-		SourceComponent->ApplyGameplayEffectToTarget(ModBuffEffect, SourceComponent, 1.f, Qualifier);
+		SourceComponent->ApplyGameplayEffectToTarget(ModBuffEffect, SourceComponent, 1.f);
 	}
 
 	{
@@ -5725,11 +5520,8 @@ bool GameplayEffectsTest_ChanceToExecuteOnOutgoingGEImmunity(UWorld *World, FAut
 		ModBuffEffect->Duration.SetValue(UGameplayEffect::INFINITE_DURATION);
 		ModBuffEffect->ChanceToExecuteOnGameplayEffect.SetValue(1.f);
 
-		FModifierQualifier Qualifier = FModifierQualifier();
-		Qualifier.PredictionKey(TEST_PREDICTIONKEY);
-
 		// Apply to self
-		SourceComponent->ApplyGameplayEffectToTarget(ModBuffEffect, SourceComponent, 1.f, Qualifier);
+		SourceComponent->ApplyGameplayEffectToTarget(ModBuffEffect, SourceComponent, 1.f);
 	}
 
 	{
@@ -5791,11 +5583,8 @@ bool GameplayEffectsTest_ChanceToExecuteOnIncomingGEMod(UWorld *World, FAutomati
 		BuffDmgEffect->Duration.SetValue(UGameplayEffect::INFINITE_DURATION);
 		BuffDmgEffect->ChanceToExecuteOnGameplayEffect.SetValue(0.f);
 
-		FModifierQualifier Qualifier = FModifierQualifier();
-		Qualifier.PredictionKey(TEST_PREDICTIONKEY);
-
 		// Apply to self
-		SourceComponent->ApplyGameplayEffectToTarget(BuffDmgEffect, DestComponent, 1.f, Qualifier);
+		SourceComponent->ApplyGameplayEffectToTarget(BuffDmgEffect, DestComponent, 1.f);
 	}
 
 	{
@@ -5835,11 +5624,8 @@ bool GameplayEffectsTest_ChanceToExecuteOnIncomingGEMod(UWorld *World, FAutomati
 		BuffDmgEffect->Duration.SetValue(UGameplayEffect::INFINITE_DURATION);
 		BuffDmgEffect->ChanceToExecuteOnGameplayEffect.SetValue(1.f);
 
-		FModifierQualifier Qualifier = FModifierQualifier();
-		Qualifier.PredictionKey(TEST_PREDICTIONKEY);
-
 		// Apply to self
-		SourceComponent->ApplyGameplayEffectToTarget(BuffDmgEffect, DestComponent, 1.f, Qualifier);
+		SourceComponent->ApplyGameplayEffectToTarget(BuffDmgEffect, DestComponent, 1.f);
 	}
 
 	{
@@ -5895,11 +5681,8 @@ bool GameplayEffectsTest_ChanceToExecuteOnIncomingGEImmunity(UWorld *World, FAut
 		ModBuffEffect->Duration.SetValue(UGameplayEffect::INFINITE_DURATION);
 		ModBuffEffect->ChanceToExecuteOnGameplayEffect.SetValue(0.f);
 
-		FModifierQualifier Qualifier = FModifierQualifier();
-		Qualifier.PredictionKey(TEST_PREDICTIONKEY);
-
 		// Apply to self
-		SourceComponent->ApplyGameplayEffectToTarget(ModBuffEffect, DestComponent, 1.f, Qualifier);
+		SourceComponent->ApplyGameplayEffectToTarget(ModBuffEffect, DestComponent, 1.f);
 	}
 
 	{
@@ -5933,11 +5716,8 @@ bool GameplayEffectsTest_ChanceToExecuteOnIncomingGEImmunity(UWorld *World, FAut
 		ModBuffEffect->Duration.SetValue(UGameplayEffect::INFINITE_DURATION);
 		ModBuffEffect->ChanceToExecuteOnGameplayEffect.SetValue(1.f);
 
-		FModifierQualifier Qualifier = FModifierQualifier();
-		Qualifier.PredictionKey(TEST_PREDICTIONKEY);
-
 		// Apply to self
-		SourceComponent->ApplyGameplayEffectToTarget(ModBuffEffect, DestComponent, 1.f, Qualifier);
+		SourceComponent->ApplyGameplayEffectToTarget(ModBuffEffect, DestComponent, 1.f);
 	}
 
 	{
@@ -5993,11 +5773,8 @@ bool GameplayEffectsTest_ModifyChanceToApplyToTarget(UWorld *World, FAutomationT
 		BaseEffect->Modifiers[0].EffectType = EGameplayModEffect::ChanceApplyTarget;
 		BaseEffect->Duration.SetValue(UGameplayEffect::INFINITE_DURATION);
 
-		FModifierQualifier Qualifier = FModifierQualifier();
-		Qualifier.PredictionKey(TEST_PREDICTIONKEY);
-
 		// Apply to self
-		SourceComponent->ApplyGameplayEffectToTarget(BaseEffect, SourceComponent, 0.f, Qualifier);
+		SourceComponent->ApplyGameplayEffectToTarget(BaseEffect, SourceComponent, 0.f);
 	}
 
 	{
@@ -6058,11 +5835,8 @@ bool GameplayEffectsTest_ModifyChanceToExecuteOnGE(UWorld *World, FAutomationTes
 		BuffEffect->GameplayEffectTags.AddTag(IGameplayTagsModule::RequestGameplayTag(FName(TEXT("Damage.Buff"))));
 		BuffEffect->ChanceToExecuteOnGameplayEffect.SetValue(0.f);
 
-		FModifierQualifier Qualifier = FModifierQualifier();
-		Qualifier.PredictionKey(TEST_PREDICTIONKEY);
-
 		// Apply to self
-		SourceComponent->ApplyGameplayEffectToTarget(BuffEffect, SourceComponent, 0.f, Qualifier);
+		SourceComponent->ApplyGameplayEffectToTarget(BuffEffect, SourceComponent, 0.f);
 	}
 
 	// verify that outgoing damage is unbuffed
@@ -6098,11 +5872,8 @@ bool GameplayEffectsTest_ModifyChanceToExecuteOnGE(UWorld *World, FAutomationTes
 		BaseEffect->Modifiers[0].EffectType = EGameplayModEffect::ChanceExecuteEffect;
 		BaseEffect->Duration.SetValue(UGameplayEffect::INFINITE_DURATION);
 
-		FModifierQualifier Qualifier = FModifierQualifier();
-		Qualifier.PredictionKey(TEST_PREDICTIONKEY);
-
 		// Apply to self
-		SourceComponent->ApplyGameplayEffectToTarget(BaseEffect, SourceComponent, 0.f, Qualifier);
+		SourceComponent->ApplyGameplayEffectToTarget(BaseEffect, SourceComponent, 0.f);
 	}
 
 	// verify that the buff applies now
