SkillSystem - cleanup and pass on ActorInfo vs new ActivationInfo state

diff --git a/Engine/Source/Editor/SkillSystemEditor/Private/AttributeDetails.cpp b/Engine/Source/Editor/SkillSystemEditor/Private/AttributeDetails.cpp
--- a/Engine/Source/Editor/SkillSystemEditor/Private/AttributeDetails.cpp
+++ b/Engine/Source/Editor/SkillSystemEditor/Private/AttributeDetails.cpp
@@ -362,7 +362,7 @@ UCurveTable * FScalableFloatDetails::GetCurveTable()
 
 	if (CurveTable == NULL)
 	{
-		CurveTable = ISkillSystemModule::Get().GetSkillSystemGlobals().GetGlobalCurveTable();
+		CurveTable = ISkillSystemModule::Get().GetSkillSystemGlobals()->GetGlobalCurveTable();
 	}
 
 	return CurveTable;
diff --git a/Engine/Source/Runtime/SkillSystem/Classes/Abilities/GameplayAbility.h b/Engine/Source/Runtime/SkillSystem/Classes/Abilities/GameplayAbility.h
--- a/Engine/Source/Runtime/SkillSystem/Classes/Abilities/GameplayAbility.h
+++ b/Engine/Source/Runtime/SkillSystem/Classes/Abilities/GameplayAbility.h
@@ -105,6 +105,10 @@ private:
 	int32 Handle;
 };
 
+/**
+ *	Not implemented yet, but we will need something like this to track ability + levels if we choose to support ability leveling in the base classes.
+ */
+
 USTRUCT()
 struct FGameplayAbilitySpec
 {
@@ -133,7 +137,9 @@ class SKILLSYSTEM_API UGameplayAbility : public UDataAsset
 public:
 
 	UFUNCTION(Client, Reliable)
-	void ClientActivateAbilitySucceed(AActor *Actor);
+	void ClientActivateAbilitySucceed(int32 PredictionKey);
+
+	virtual void ClientActivateAbilitySucceed_Internal(int32 PredictionKey);
 
 	bool IsSupportedForNetworking() const OVERRIDE;
 	
@@ -165,22 +171,20 @@ public:
 	 *	-This function handles networking and prediction
 	 *	-If all goes well, CallActivateAbility is called next.
 	 */
-	bool TryActivateAbility(const FGameplayAbilityActorInfo ActorInfo, UGameplayAbility ** OutInstancedAbility = NULL);
+	bool TryActivateAbility(const FGameplayAbilityActorInfo* ActorInfo, int32 PredictionKey = 0, UGameplayAbility ** OutInstancedAbility = NULL);
 	
-	/** Returns true if this ability can be activated right now. Has no side effects! */
-	virtual bool CanActivateAbility(const FGameplayAbilityActorInfo ActorInfo) const;		
+	/** Returns true if this ability can be activated right now. Has no side effects */
+	virtual bool CanActivateAbility(const FGameplayAbilityActorInfo* ActorInfo) const;
 
 	// -----------------------------------------------
 
 	/** Input binding. Base implementation calls TryActivateAbility */
-	virtual void InputPressed(int32 InputID, const FGameplayAbilityActorInfo ActorInfo);
+	virtual void InputPressed(int32 InputID, const FGameplayAbilityActorInfo* ActorInfo);
 
 	/** Input binding. Base implementation does nothing */
-	virtual void InputReleased(int32 InputID, const FGameplayAbilityActorInfo ActorInfo);
-
-	UFUNCTION(BlueprintCallable, Category=Ability)
-	virtual void EndAbility(const FGameplayAbilityActorInfo ActorInfo);
-
+	virtual void InputReleased(int32 InputID, const FGameplayAbilityActorInfo* ActorInfo);
+	
+	virtual void EndAbility(const FGameplayAbilityActorInfo* ActorInfo);
 
 	// ----------------------------------------------------------------------------------------------------------------
 	//
@@ -188,7 +192,7 @@ public:
 	//
 	// ----------------------------------------------------------------------------------------------------------------
 
-	float GetCooldownTimeRemaining(const FGameplayAbilityActorInfo ActorInfo) const;
+	float GetCooldownTimeRemaining(const FGameplayAbilityActorInfo* ActorInfo) const;
 
 	// ----------------------------------------------------------------------------------------------------------------
 	//
@@ -227,17 +231,17 @@ public:
 	//
 	// ----------------------------------------------------------------------------------------------------------------
 
-	void MontageBranchPoint_AbilityDecisionStop(const FGameplayAbilityActorInfo ActorInfo) const;
+	void MontageBranchPoint_AbilityDecisionStop(const FGameplayAbilityActorInfo* ActorInfo) const;
 
-	void MontageBranchPoint_AbilityDecisionStart(const FGameplayAbilityActorInfo ActorInfo) const;
+	void MontageBranchPoint_AbilityDecisionStart(const FGameplayAbilityActorInfo* ActorInfo) const;
 
 	// ----------------------------------------------------------------------------------------------------------------
 
-	virtual void PredictiveActivateAbility(const FGameplayAbilityActorInfo ActorInfo);
+	virtual void PredictiveActivateAbility(const FGameplayAbilityActorInfo* ActorInfo, const FGameplayAbilityActivationInfo ActivationInfo);
 
-	virtual void ServerTryActivateAbility(const FGameplayAbilityActorInfo ActorInfo);
+	virtual void ServerTryActivateAbility(const FGameplayAbilityActorInfo* ActorInfo, const FGameplayAbilityActivationInfo ActivationInfo);
 
-	virtual void ClientActivateAbility(const FGameplayAbilityActorInfo ActorInfo);
+	virtual void ClientActivateAbility(const FGameplayAbilityActorInfo* ActorInfo, const FGameplayAbilityActivationInfo ActivationInfo);
 
 	// ----------------------------------------------------------------------------------------------------------------
 
@@ -266,23 +270,21 @@ protected:
 	 *  -This function must call CommitAbility()
 	 */
 
-	// FIXME: Making this blueprintcallable is danagerous!
-	UFUNCTION(BlueprintCallable, Category = "Ability")
-	virtual void ActivateAbility(const FGameplayAbilityActorInfo ActorInfo);
+	virtual void ActivateAbility(const FGameplayAbilityActorInfo * ActorInfo, const FGameplayAbilityActivationInfo ActivationInfo);
 
 	/** 
 	 * Attempts to commit the ability (spend resources, etc). This our last chance to fail.
 	 *	-Child classes that override ActivateAbility must call this themselves!
 	 */
-	UFUNCTION(BlueprintCallable, Category = "Ability")
-	virtual bool CommitAbility(const FGameplayAbilityActorInfo ActorInfo);
+
+	virtual bool CommitAbility(const FGameplayAbilityActorInfo* ActorInfo, const FGameplayAbilityActivationInfo ActivationInfo);
 	
 	/** Do boilerplate init stuff and then call ActivateAbility */
-	void PreActivate(const FGameplayAbilityActorInfo ActorInfo);
+	void PreActivate(const FGameplayAbilityActorInfo* ActorInfo, const FGameplayAbilityActivationInfo ActivationInfo);
 
-	void CallActivateAbility(const FGameplayAbilityActorInfo ActorInfo);
+	void CallActivateAbility(const FGameplayAbilityActorInfo* ActorInfo, const FGameplayAbilityActivationInfo ActivationInfo);
 
-	void CallPredictiveActivateAbility(const FGameplayAbilityActorInfo ActorInfo);
+	void CallPredictiveActivateAbility(const FGameplayAbilityActorInfo* ActorInfo, const FGameplayAbilityActivationInfo ActivationInfo);
 
 	UPROPERTY(EditDefaultsOnly, Category = Tags)
 	bool Prediction;
@@ -291,13 +293,13 @@ protected:
 	 * The last chance to fail before commiting
 	 *	-This will usually be the same as CanActivateAbility. Some abilities may need to do extra checks here if they are consuming extra stuff in CommitExecute
 	 */
-	virtual bool CommitCheck(const FGameplayAbilityActorInfo ActorInfo);
+	virtual bool CommitCheck(const FGameplayAbilityActorInfo* ActorInfo, const FGameplayAbilityActivationInfo ActivationInfo);
 
 	/** Does the commmit atomically (consume resources, do cooldowns, etc) */
-	virtual void CommitExecute(const FGameplayAbilityActorInfo ActorInfo);
+	virtual void CommitExecute(const FGameplayAbilityActorInfo* ActorInfo, const FGameplayAbilityActivationInfo ActivationInfo);
 
 	/** Destroys intanced-per-execution abilities. Instance-per-actor abilities should 'reset'. Non instance abilities - what can we do? */
-	virtual void CancelAbility(const FGameplayAbilityActorInfo ActorInfo) { }
+	virtual void CancelAbility(const FGameplayAbilityActorInfo* ActorInfo, const FGameplayAbilityActivationInfo ActivationInfo) { }
 
 	// --------------------------------------------------------------------------------
 
@@ -305,14 +307,14 @@ protected:
 	// --------------------------------------------------------------------------------
 
 	/** Checks cooldown. returns true if we can be used again. False if not */
-	bool	CheckCooldown(const FGameplayAbilityActorInfo ActorInfo) const;
+	bool	CheckCooldown(const FGameplayAbilityActorInfo* ActorInfo) const;
 
 	/** Applies CooldownGameplayEffect to the target */
-	void	ApplyCooldown(const FGameplayAbilityActorInfo ActorInfo);
+	void	ApplyCooldown(const FGameplayAbilityActorInfo* ActorInfo, const FGameplayAbilityActivationInfo ActivationInfo);
 
-	bool	CheckCost(const FGameplayAbilityActorInfo ActorInfo) const;
+	bool	CheckCost(const FGameplayAbilityActorInfo* ActorInfo) const;
 
-	void	ApplyCost(const FGameplayAbilityActorInfo ActorInfo);
+	void	ApplyCost(const FGameplayAbilityActorInfo* ActorInfo, const FGameplayAbilityActivationInfo ActivationInfo);
 
 	friend class UAttributeComponent;
 };
\ No newline at end of file
diff --git a/Engine/Source/Runtime/SkillSystem/Classes/Abilities/GameplayAbility_CharacterJump.h b/Engine/Source/Runtime/SkillSystem/Classes/Abilities/GameplayAbility_CharacterJump.h
--- a/Engine/Source/Runtime/SkillSystem/Classes/Abilities/GameplayAbility_CharacterJump.h
+++ b/Engine/Source/Runtime/SkillSystem/Classes/Abilities/GameplayAbility_CharacterJump.h
@@ -15,13 +15,13 @@ class SKILLSYSTEM_API UGameplayAbility_CharacterJump : public UGameplayAbility
 
 public:
 
-	virtual bool CanActivateAbility(const FGameplayAbilityActorInfo ActorInfo) const OVERRIDE;
+	virtual bool CanActivateAbility(const FGameplayAbilityActorInfo* ActorInfo) const OVERRIDE;
 	
-	virtual void ActivateAbility(const FGameplayAbilityActorInfo OwnerInfo) OVERRIDE;
+	virtual void ActivateAbility(const FGameplayAbilityActorInfo* OwnerInfo, const FGameplayAbilityActivationInfo ActivationInfo) OVERRIDE;
 
-	virtual void InputReleased(int32 InputID, const FGameplayAbilityActorInfo ActorInfo) OVERRIDE;
+	virtual void InputReleased(int32 InputID, const FGameplayAbilityActorInfo* ActorInfo) OVERRIDE;
 
-	virtual void PredictiveActivateAbility(const FGameplayAbilityActorInfo ActorInfo) OVERRIDE;
+	virtual void PredictiveActivateAbility(const FGameplayAbilityActorInfo* ActorInfo, const FGameplayAbilityActivationInfo ActivationInfo) OVERRIDE;
 
-	virtual void CancelAbility(const FGameplayAbilityActorInfo ActorInfo) OVERRIDE;
+	virtual void CancelAbility(const FGameplayAbilityActorInfo* ActorInfo, const FGameplayAbilityActivationInfo ActivationInfo) OVERRIDE;
 };
\ No newline at end of file
diff --git a/Engine/Source/Runtime/SkillSystem/Classes/Abilities/GameplayAbility_Instanced.h b/Engine/Source/Runtime/SkillSystem/Classes/Abilities/GameplayAbility_Instanced.h
--- a/Engine/Source/Runtime/SkillSystem/Classes/Abilities/GameplayAbility_Instanced.h
+++ b/Engine/Source/Runtime/SkillSystem/Classes/Abilities/GameplayAbility_Instanced.h
@@ -17,39 +17,56 @@ class SKILLSYSTEM_API UGameplayAbility_Instanced : public UGameplayAbility
 	GENERATED_UCLASS_BODY()
 
 public:
+
+	virtual void ClientActivateAbilitySucceed_Internal(int32 PredictionKey) OVERRIDE;
 	
-	virtual void EndAbility(const FGameplayAbilityActorInfo ActorInfo) OVERRIDE;
+	virtual void EndAbility(const FGameplayAbilityActorInfo* ActorInfo) OVERRIDE;
 	
 	// -------------------------------------------------
 
 	UFUNCTION(BlueprintImplementableEvent, Category = Ability)
-	virtual bool K2_CanActivateAbility(const FGameplayAbilityActorInfo ActorInfo) const;
+	virtual bool K2_CanActivateAbility() const;
 
-	virtual bool CanActivateAbility(const FGameplayAbilityActorInfo ActorInfo) const OVERRIDE;
+	virtual bool CanActivateAbility(const FGameplayAbilityActorInfo* ActorInfo) const OVERRIDE;
 
 	// -------------------------------------------------
 
 	UFUNCTION(BlueprintImplementableEvent, Category = Ability)
-	virtual void K2_CommitExecute(const FGameplayAbilityActorInfo ActorInfo);
+	virtual void K2_CommitExecute();
 
-	virtual void CommitExecute(const FGameplayAbilityActorInfo ActorInfo) OVERRIDE;
+	virtual void CommitExecute(const FGameplayAbilityActorInfo* ActorInfo, const FGameplayAbilityActivationInfo ActivationInfo) OVERRIDE;
 
 	// -------------------------------------------------
 
 	UFUNCTION(BlueprintImplementableEvent, Category = Ability)
-	virtual void K2_ActivateAbility(const FGameplayAbilityActorInfo ActorInfo);
+	virtual void K2_ActivateAbility();
 
-	virtual void ActivateAbility(const FGameplayAbilityActorInfo ActorInfo) OVERRIDE;
+	virtual void ActivateAbility(const FGameplayAbilityActorInfo* ActorInfo, const FGameplayAbilityActivationInfo ActivationInfo) OVERRIDE;
 
 	// -------------------------------------------------
 
 	UFUNCTION(BlueprintImplementableEvent, Category = Ability)
-	virtual void K2_PredictiveActivateAbility(const FGameplayAbilityActorInfo ActorInfo);
+	virtual void K2_PredictiveActivateAbility();
+
+	virtual void PredictiveActivateAbility(const FGameplayAbilityActorInfo* ActorInfo, const FGameplayAbilityActivationInfo ActivationInfo) OVERRIDE;
+
+	// -------------------------------------------------
+
+	UFUNCTION(BlueprintPure, Category = Ability)
+	FGameplayAbilityActorInfo GetActorInfo();
+
+	// -------------------------------------------------
+
+	UFUNCTION(BlueprintCallable, Category = Ability)
+	virtual bool K2_CommitAbility();
 
-	virtual void PredictiveActivateAbility(const FGameplayAbilityActorInfo ActorInfo) OVERRIDE;
+	UFUNCTION(BlueprintCallable, Category = Ability)
+	virtual void K2_EndAbility();
+	
 
 	// -------------------------------------------------
 
+
 	virtual EGameplayAbilityInstancingPolicy::Type GetInstancingPolicy() const OVERRIDE
 	{
 		return InstancedPerExecution ? EGameplayAbilityInstancingPolicy::InstancedPerExecution : EGameplayAbilityInstancingPolicy::InstancedPerActor;
@@ -70,6 +87,19 @@ public:
 
 	bool CallRemoteFunction(UFunction* Function, void* Parameters, FOutParmRec* OutParms, FFrame* Stack);
 
-private:
-	
+
+	// -------------------------------------------------
+
+
+	/** This is shared, cached information about the thing using us
+	 *		E.g, Actor*, MovementComponent*, AnimInstance, etc.
+	 *		This is hopefully allocated once per actor and shared by many abilities.
+	 *		The actual struct may be overridden per game to include game specific data.
+	 *		(E.g, child classes may want to cast to FMyGameAbilityActorInfo)
+	 */
+	mutable const FGameplayAbilityActorInfo * CurrentActorInfo;
+
+	/** This is information specific to this instance of the ability. E.g, whether it is predicting, authorting, confirmed, etc. */
+	UPROPERTY(BlueprintReadOnly, Category=Ability)
+	FGameplayAbilityActivationInfo	CurrentActivationInfo;
 };
\ No newline at end of file
diff --git a/Engine/Source/Runtime/SkillSystem/Classes/Abilities/GameplayAbility_Montage.h b/Engine/Source/Runtime/SkillSystem/Classes/Abilities/GameplayAbility_Montage.h
--- a/Engine/Source/Runtime/SkillSystem/Classes/Abilities/GameplayAbility_Montage.h
+++ b/Engine/Source/Runtime/SkillSystem/Classes/Abilities/GameplayAbility_Montage.h
@@ -15,7 +15,7 @@ class SKILLSYSTEM_API UGameplayAbility_Montage : public UGameplayAbility
 
 public:
 	
-	virtual void ActivateAbility(const FGameplayAbilityActorInfo OwnerInfo) OVERRIDE;
+	virtual void ActivateAbility(const FGameplayAbilityActorInfo* OwnerInfo, const FGameplayAbilityActivationInfo ActivationInfo) OVERRIDE;
 
 	UPROPERTY(EditDefaultsOnly, Category = MontageAbility)
 	UAnimMontage *	MontageToPlay;
@@ -30,5 +30,5 @@ public:
 	UPROPERTY(EditDefaultsOnly, Category = MontageAbility)
 	TArray<const UGameplayEffect*>	GameplayEffectsWhileAnimating;
 
-	void OnMontageEnded(UAnimMontage* Montage, bool bInterrupted, const FGameplayAbilityActorInfo OwnerInfo, TArray<struct FActiveGameplayEffectHandle>	AppliedEffects);
+	void OnMontageEnded(UAnimMontage* Montage, bool bInterrupted, const FGameplayAbilityActorInfo* OwnerInfo, TArray<struct FActiveGameplayEffectHandle>	AppliedEffects);
 };
\ No newline at end of file
diff --git a/Engine/Source/Runtime/SkillSystem/Classes/AttributeComponent.h b/Engine/Source/Runtime/SkillSystem/Classes/AttributeComponent.h
--- a/Engine/Source/Runtime/SkillSystem/Classes/AttributeComponent.h
+++ b/Engine/Source/Runtime/SkillSystem/Classes/AttributeComponent.h
@@ -234,7 +234,7 @@ class SKILLSYSTEM_API UAttributeComponent : public UActorComponent
 	 *	GameplayAbilities
 	 *	
 	 *	The role of the AttributeComponent wrt Abilities is to provide:
-	 *		-Management of abiltiy instances (whether per actor or per instance).
+	 *		-Management of abiltiy instances (whether per actor or per execution instance).
 	 *			-Someone *has* to keep track of these instances.
 	 *			-Non instanced abilities *could* be executed without any ability stuff in attribute component.
 	 *				They should be able to operate on an GameplayAbilityActorInfo + GameplayAbility.
@@ -245,9 +245,9 @@ class SKILLSYSTEM_API UAttributeComponent : public UActorComponent
 	 *	
 	 */
 
-	class UGameplayAbility * CreateNewInstanceOfAbility(UGameplayAbility *Ability);
+	class UGameplayAbility * CreateNewInstanceOfAbility(UGameplayAbility* Ability);
 
-	void CancelAbilitiesWithTags(const FGameplayTagContainer Tags, const FGameplayAbilityActorInfo & ActorInfo, UGameplayAbility * Ignore);
+	void CancelAbilitiesWithTags(const FGameplayTagContainer Tags, const FGameplayAbilityActorInfo* ActorInfo, const FGameplayAbilityActivationInfo ActivationInfo, UGameplayAbility* Ignore);
 	
 	/** References to non-replicating abilities that we instanced. We need to keep these refernces to avoid GC */
 	UPROPERTY()
@@ -271,13 +271,13 @@ class SKILLSYSTEM_API UAttributeComponent : public UActorComponent
 	
 
 	UFUNCTION(Server, WithValidation, Reliable)
-	void	ServerTryActivateAbility(class UGameplayAbility * AbilityToActivate, int32 PredictionKey);
+	void	ServerTryActivateAbility(class UGameplayAbility* AbilityToActivate, int32 PredictionKey);
 
 	UFUNCTION(Client, Reliable)
-	void	ClientActivateAbilityFailed(class UGameplayAbility * AbilityToActivate, int32 PredictionKey);
+	void	ClientActivateAbilityFailed(class UGameplayAbility* AbilityToActivate, int32 PredictionKey);
 
 	UFUNCTION(Client, Reliable)
-	void	ClientActivateAbilitySucceed(class UGameplayAbility * AbilityToActivate, int32 PredictionKey);
+	void	ClientActivateAbilitySucceed(class UGameplayAbility* AbilityToActivate, int32 PredictionKey);
 
 	// ----------------------------------------------------------------------------------------------------------------
 
@@ -296,8 +296,11 @@ class SKILLSYSTEM_API UAttributeComponent : public UActorComponent
 	
 
 	// -----------------------------------------------------------------------------
-	
-	
+
+	/** Cached off data about the owning actor that abilities will need to frequently access (movement component, mesh component, anim instance, etc) */
+	TSharedPtr<FGameplayAbilityActorInfo>	AbilityActorInfo;
+
+	void InitAbilityActorInfo();
 
 private:
 
diff --git a/Engine/Source/Runtime/SkillSystem/Classes/SkillSystemGlobals.h b/Engine/Source/Runtime/SkillSystem/Classes/SkillSystemGlobals.h
--- a/Engine/Source/Runtime/SkillSystem/Classes/SkillSystemGlobals.h
+++ b/Engine/Source/Runtime/SkillSystem/Classes/SkillSystemGlobals.h
@@ -6,6 +6,7 @@
 
 class AActor;
 class UAttributeComponent;
+struct FGameplayAbilityActorInfo;
 
 /** Holds global data for the skill system. Can be configured per project via config file */
 UCLASS(config=Game)
@@ -35,12 +36,16 @@ class SKILLSYSTEM_API USkillSystemGlobals : public UObject
 		GlobalAttributeDataTable = InTable;
 	}
 
+	static USkillSystemGlobals&	Get();
+
 	/**
 	 *	Games may want to override this in a UMyProjectSkillSystemsGlobals class and provide
 	 *	a faster lookup for attribute compnents (E.g. Cast<MyProjPawn>(Actor)->AttributeComponent; etc)
 	 */
 	virtual UAttributeComponent * GetAttributeComponentFromActor(AActor *Actor) const;
 
+	virtual FGameplayAbilityActorInfo * AllocAbilityActorInfo(AActor *Actor) const;
+
 private:
 
 	class UCurveTable* GlobalCurveTable;
diff --git a/Engine/Source/Runtime/SkillSystem/Private/Abilities/GameplayAbility.cpp b/Engine/Source/Runtime/SkillSystem/Private/Abilities/GameplayAbility.cpp
--- a/Engine/Source/Runtime/SkillSystem/Private/Abilities/GameplayAbility.cpp
+++ b/Engine/Source/Runtime/SkillSystem/Private/Abilities/GameplayAbility.cpp
@@ -38,7 +38,7 @@ bool UGameplayAbility::IsSupportedForNetworking() const
 	return Supported;
 }
 
-bool UGameplayAbility::CanActivateAbility(const FGameplayAbilityActorInfo ActorInfo) const
+bool UGameplayAbility::CanActivateAbility(const FGameplayAbilityActorInfo* ActorInfo) const
 {
 	// Check cooldowns and resource consumption immediately
 	if (!CheckCooldown(ActorInfo))
@@ -54,37 +54,39 @@ bool UGameplayAbility::CanActivateAbility(const FGameplayAbilityActorInfo ActorI
 	return true;
 }
 
-bool UGameplayAbility::CommitAbility(const FGameplayAbilityActorInfo ActorInfo)
+bool UGameplayAbility::CommitAbility(const FGameplayAbilityActorInfo* ActorInfo, const FGameplayAbilityActivationInfo ActivationInfo)
 {
-	if (!CommitCheck(ActorInfo))
+	if (!CommitCheck(ActorInfo, ActivationInfo))
 	{
 		return false;
 	}
 
-	CommitExecute(ActorInfo);
+	CommitExecute(ActorInfo, ActivationInfo);
 	return true;
 }
 
-bool UGameplayAbility::CommitCheck(const FGameplayAbilityActorInfo ActorInfo)
+bool UGameplayAbility::CommitCheck(const FGameplayAbilityActorInfo* ActorInfo, const FGameplayAbilityActivationInfo ActivationInfo)
 {
 	return CanActivateAbility(ActorInfo);
 }
 
-void UGameplayAbility::CommitExecute(const FGameplayAbilityActorInfo ActorInfo)
+void UGameplayAbility::CommitExecute(const FGameplayAbilityActorInfo* ActorInfo, const FGameplayAbilityActivationInfo ActivationInfo)
 {
-	if (!ActorInfo.AttributeComponent->IsNetSimulating() || ActorInfo.PredictionKey != 0)
+	if (!ActorInfo->AttributeComponent->IsNetSimulating() || ActivationInfo.PredictionKey != 0)
 	{
-		ApplyCooldown(ActorInfo);
+		ApplyCooldown(ActorInfo, ActivationInfo);
 
-		ApplyCost(ActorInfo);
+		ApplyCost(ActorInfo, ActivationInfo);
 	}
 }
 
-bool UGameplayAbility::TryActivateAbility(const FGameplayAbilityActorInfo ActorInfo, UGameplayAbility ** OutInstancedAbility)
+bool UGameplayAbility::TryActivateAbility(const FGameplayAbilityActorInfo* ActorInfo, int32 PredictionKey, UGameplayAbility ** OutInstancedAbility)
 {	
 	// This should only come from button presses/local instigation
-	ENetRole NetMode = ActorInfo.Actor->Role;
+	ENetRole NetMode = ActorInfo->Actor->Role;
 	ensure(NetMode != ROLE_SimulatedProxy);
+	
+	FGameplayAbilityActivationInfo	ActivationInfo(ActorInfo->Actor.Get(), PredictionKey);
 
 	// Always do a non instanced CanActivate check
 	if (!CanActivateAbility(ActorInfo))
@@ -99,8 +101,8 @@ bool UGameplayAbility::TryActivateAbility(const FGameplayAbilityActorInfo ActorI
 		
 		if (GetInstancingPolicy() == EGameplayAbilityInstancingPolicy::InstancedPerExecution)
 		{
-			UGameplayAbility * Ability = ActorInfo.AttributeComponent->CreateNewInstanceOfAbility(this);
-			Ability->CallActivateAbility(ActorInfo);
+			UGameplayAbility * Ability = ActorInfo->AttributeComponent->CreateNewInstanceOfAbility(this);
+			Ability->CallActivateAbility(ActorInfo, ActivationInfo);
 			if (OutInstancedAbility)
 			{
 				*OutInstancedAbility = Ability;
@@ -108,17 +110,17 @@ bool UGameplayAbility::TryActivateAbility(const FGameplayAbilityActorInfo ActorI
 		}
 		else
 		{
-			CallActivateAbility(ActorInfo);
+			CallActivateAbility(ActorInfo, ActivationInfo);
 		}
 	}
 	else if (NetExecutionPolicy == EGameplayAbilityNetExecutionPolicy::Server)
 	{
-		ServerTryActivateAbility(ActorInfo);
+		ServerTryActivateAbility(ActorInfo, ActivationInfo);
 	}
 	else if (NetExecutionPolicy == EGameplayAbilityNetExecutionPolicy::Predictive)
 	{
 		// This execution is now officially EGameplayAbilityActivationMode:Predicting and has a PredictionKey
-		ActorInfo.GeneratePredictionKey();
+		ActivationInfo.GeneratePredictionKey(ActorInfo->AttributeComponent.Get());
 
 		if (GetInstancingPolicy() == EGameplayAbilityInstancingPolicy::InstancedPerExecution)
 		{
@@ -128,8 +130,8 @@ bool UGameplayAbility::TryActivateAbility(const FGameplayAbilityActorInfo ActorI
 			
 			if(GetReplicationPolicy() == EGameplayAbilityReplicationPolicy::ReplicateNone)
 			{
-				UGameplayAbility * Ability = ActorInfo.AttributeComponent->CreateNewInstanceOfAbility(this);
-				Ability->CallPredictiveActivateAbility(ActorInfo);
+				UGameplayAbility * Ability = ActorInfo->AttributeComponent->CreateNewInstanceOfAbility(this);
+				Ability->CallPredictiveActivateAbility(ActorInfo, ActivationInfo);
 				if (OutInstancedAbility)
 				{
 					*OutInstancedAbility = Ability;
@@ -142,78 +144,78 @@ bool UGameplayAbility::TryActivateAbility(const FGameplayAbilityActorInfo ActorI
 		}
 		else
 		{
-			CallPredictiveActivateAbility(ActorInfo);
+			CallPredictiveActivateAbility(ActorInfo, ActivationInfo);
 		}
 		
-		ServerTryActivateAbility(ActorInfo);
+		ServerTryActivateAbility(ActorInfo, ActivationInfo);
 	}
 
 	return true;	
 }
 
-void UGameplayAbility::EndAbility(const FGameplayAbilityActorInfo ActorInfo)
+void UGameplayAbility::EndAbility(const FGameplayAbilityActorInfo* ActorInfo)
 {
 
 }
 
-void UGameplayAbility::ActivateAbility(FGameplayAbilityActorInfo ActorInfo)
+void UGameplayAbility::ActivateAbility(const FGameplayAbilityActorInfo* ActorInfo, const FGameplayAbilityActivationInfo ActivationInfo)
 {
 	// Child classes may want to do stuff here...
 
-	if (CommitAbility(ActorInfo))		// ..then commit the ability...
+	if (CommitAbility(ActorInfo, ActivationInfo))		// ..then commit the ability...
 	{
 		//	Then do more stuff...
 	}
 }
 
-void UGameplayAbility::PreActivate(const FGameplayAbilityActorInfo ActorInfo)
+void UGameplayAbility::PreActivate(const FGameplayAbilityActorInfo* ActorInfo, const FGameplayAbilityActivationInfo ActivationInfo)
 {
-	ActorInfo.AttributeComponent->CancelAbilitiesWithTags(CancelAbilitiesWithTag, ActorInfo, this);
+	ActorInfo->AttributeComponent->CancelAbilitiesWithTags(CancelAbilitiesWithTag, ActorInfo, ActivationInfo, this);
 }
 
-void UGameplayAbility::CallActivateAbility(const FGameplayAbilityActorInfo ActorInfo)
+void UGameplayAbility::CallActivateAbility(const FGameplayAbilityActorInfo* ActorInfo, const FGameplayAbilityActivationInfo ActivationInfo)
 {
-	PreActivate(ActorInfo);
-	ActivateAbility(ActorInfo);
+	PreActivate(ActorInfo, ActivationInfo);
+	ActivateAbility(ActorInfo, ActivationInfo);
 }
 
-void UGameplayAbility::CallPredictiveActivateAbility(const FGameplayAbilityActorInfo ActorInfo)
+void UGameplayAbility::CallPredictiveActivateAbility(const FGameplayAbilityActorInfo* ActorInfo, const FGameplayAbilityActivationInfo ActivationInfo)
 {
-	PreActivate(ActorInfo);
-	PredictiveActivateAbility(ActorInfo);
+	PreActivate(ActorInfo, ActivationInfo);
+	PredictiveActivateAbility(ActorInfo, ActivationInfo);
 }
 
-void UGameplayAbility::PredictiveActivateAbility(const FGameplayAbilityActorInfo ActorInfo)
+void UGameplayAbility::PredictiveActivateAbility(const FGameplayAbilityActorInfo* ActorInfo, const FGameplayAbilityActivationInfo ActivationInfo)
 {
 	
 }
 
-void UGameplayAbility::ServerTryActivateAbility(const FGameplayAbilityActorInfo ActorInfo)
+void UGameplayAbility::ServerTryActivateAbility(const FGameplayAbilityActorInfo* ActorInfo, const FGameplayAbilityActivationInfo ActivationInfo)
 {
-	ActorInfo.AttributeComponent->ServerTryActivateAbility(this, ActorInfo.PredictionKey);
+	ActorInfo->AttributeComponent->ServerTryActivateAbility(this, ActivationInfo.PredictionKey);
 }
 
-void UGameplayAbility::ClientActivateAbility(const FGameplayAbilityActorInfo ActorInfo)
-{
-	CallActivateAbility(ActorInfo);
+void UGameplayAbility::ClientActivateAbility(const FGameplayAbilityActorInfo* ActorInfo, const FGameplayAbilityActivationInfo ActivationInfo)
+{ 
+	CallActivateAbility(ActorInfo, ActivationInfo);
 }
 
-void UGameplayAbility::InputPressed(int32 InputID, const FGameplayAbilityActorInfo ActorInfo)
+void UGameplayAbility::InputPressed(int32 InputID, const FGameplayAbilityActorInfo* ActorInfo)
 {
 	TryActivateAbility(ActorInfo);
 }
 
-void UGameplayAbility::InputReleased(int32 InputID, const FGameplayAbilityActorInfo ActorInfo)
+void UGameplayAbility::InputReleased(int32 InputID, const FGameplayAbilityActorInfo* ActorInfo)
 {
 	SKILL_LOG(Log, TEXT("InputReleased: %d"), InputID);
 }
 
-bool UGameplayAbility::CheckCooldown(const FGameplayAbilityActorInfo ActorInfo) const
+bool UGameplayAbility::CheckCooldown(const FGameplayAbilityActorInfo* ActorInfo) const
 {
 	if (CooldownGameplayEffect)
 	{
-		check(ActorInfo.AttributeComponent.IsValid());
-		if (ActorInfo.AttributeComponent->HasAnyTags(CooldownGameplayEffect->OwnedTagsContainer))
+		check(ActorInfo->AttributeComponent.IsValid());
+		if (ActorInfo->AttributeComponent->HasAnyTags(CooldownGameplayEffect->OwnedTagsContainer))
 		{
 			return false;
 		}
@@ -221,40 +223,42 @@ bool UGameplayAbility::CheckCooldown(const FGameplayAbilityActorInfo ActorInfo)
 	return true;
 }
 
-void UGameplayAbility::ApplyCooldown(const FGameplayAbilityActorInfo ActorInfo)
+void UGameplayAbility::ApplyCooldown(const FGameplayAbilityActorInfo* ActorInfo, const FGameplayAbilityActivationInfo ActivationInfo)
 {
 	if (CooldownGameplayEffect)
 	{
-		check(ActorInfo.AttributeComponent.IsValid());
-		ActorInfo.AttributeComponent->ApplyGameplayEffectToSelf(CooldownGameplayEffect, 1.f, ActorInfo.Actor.Get(), FModifierQualifier().PredictionKey(ActorInfo.PredictionKey));
+		check(ActorInfo->AttributeComponent.IsValid());
+		ActorInfo->AttributeComponent->ApplyGameplayEffectToSelf(CooldownGameplayEffect, 1.f, ActorInfo->Actor.Get(), FModifierQualifier().PredictionKey(ActivationInfo.PredictionKey));
 	}
 }
 
-bool UGameplayAbility::CheckCost(const FGameplayAbilityActorInfo ActorInfo) const
+bool UGameplayAbility::CheckCost(const FGameplayAbilityActorInfo* ActorInfo) const
 {
 	if (CostGameplayEffect)
 	{
-		check(ActorInfo.AttributeComponent.IsValid());
-		return ActorInfo.AttributeComponent->CanApplyAttributeModifiers(CostGameplayEffect, 1.f, ActorInfo.Actor.Get());
+		check(ActorInfo->AttributeComponent.IsValid());
+		return ActorInfo->AttributeComponent->CanApplyAttributeModifiers(CostGameplayEffect, 1.f, ActorInfo->Actor.Get());
 	}
 	return true;
 }
 
-void UGameplayAbility::ApplyCost(const FGameplayAbilityActorInfo ActorInfo)
+void UGameplayAbility::ApplyCost(const FGameplayAbilityActorInfo* ActorInfo, const FGameplayAbilityActivationInfo ActivationInfo)
 {
-	check(ActorInfo.AttributeComponent.IsValid());
+	check(ActorInfo->AttributeComponent.IsValid());
 	if (CostGameplayEffect)
 	{
-		ActorInfo.AttributeComponent->ApplyGameplayEffectToSelf(CostGameplayEffect, 1.f, ActorInfo.Actor.Get(), FModifierQualifier().PredictionKey(ActorInfo.PredictionKey));
+		ActorInfo->AttributeComponent->ApplyGameplayEffectToSelf(CostGameplayEffect, 1.f, ActorInfo->Actor.Get(), FModifierQualifier().PredictionKey(ActivationInfo.PredictionKey));
 	}
 }
 
-float UGameplayAbility::GetCooldownTimeRemaining(const FGameplayAbilityActorInfo ActorInfo) const
+float UGameplayAbility::GetCooldownTimeRemaining(const FGameplayAbilityActorInfo* ActorInfo) const
 {
-	check(ActorInfo.AttributeComponent.IsValid());
+	SCOPE_CYCLE_COUNTER(STAT_GameplayAbilityGetCooldownTimeRemaining);
+
+	check(ActorInfo->AttributeComponent.IsValid());
 	if (CooldownGameplayEffect)
 	{
-		TArray< float > Durations = ActorInfo.AttributeComponent->GetActiveEffectsTimeRemaining(FActiveGameplayEffectQuery(&CooldownGameplayEffect->OwnedTagsContainer));
+		TArray< float > Durations = ActorInfo->AttributeComponent->GetActiveEffectsTimeRemaining(FActiveGameplayEffectQuery(&CooldownGameplayEffect->OwnedTagsContainer));
 		if (Durations.Num() > 0)
 		{
 			Durations.Sort();
@@ -267,29 +271,25 @@ float UGameplayAbility::GetCooldownTimeRemaining(const FGameplayAbilityActorInfo
 
 // --------------------------------------------------------------------
 
-void UGameplayAbility::MontageBranchPoint_AbilityDecisionStop(const FGameplayAbilityActorInfo ActorInfo) const
+void UGameplayAbility::MontageBranchPoint_AbilityDecisionStop(const FGameplayAbilityActorInfo* ActorInfo) const
 {
 
 }
 
-void UGameplayAbility::MontageBranchPoint_AbilityDecisionStart(const FGameplayAbilityActorInfo ActorInfo) const
+void UGameplayAbility::MontageBranchPoint_AbilityDecisionStart(const FGameplayAbilityActorInfo* ActorInfo) const
 {
 
 }
 
-void UGameplayAbility::ClientActivateAbilitySucceed_Implementation(AActor *Actor)
+void UGameplayAbility::ClientActivateAbilitySucceed_Implementation(int32 PredictionKey)
 {
-	if (Actor)
-	{
-		FGameplayAbilityActorInfo	ActorInfo;
-		ActorInfo.InitFromActor(Actor);
+	// Child classes can't override ClientActivateAbilitySucceed_Implementation, so just call ClientActivateAbilitySucceed_Internal
+	ClientActivateAbilitySucceed_Internal(PredictionKey);
+}
 
-		CallActivateAbility(ActorInfo);
-	}
-	else
-	{
-		SKILL_LOG(Warning, TEXT("ClientActivateAbilitySucceed called on %s with no valid Actor"), *GetName());
-	}
+void UGameplayAbility::ClientActivateAbilitySucceed_Internal(int32 PredictionKey)
+{
+	SKILL_LOG(Fatal, TEXT("ClientActivateAbilitySucceed_Internal called on base class ability %s"), *GetName());
 }
 
 //----------------------------------------------------------------------
@@ -298,12 +298,6 @@ void FGameplayAbilityActorInfo::InitFromActor(AActor *InActor)
 {
 	Actor = InActor;
 
-	// On Init, we are either Authority or NonAuthority. We haven't been given a PredictionKey and we haven't been confirmed.
-	// NonAuthority essentially means 'I'm not sure what how I'm going to do this yet'.
-	// 
-	ActivationMode = (Actor->Role == ROLE_Authority ? EGameplayAbilityActivationMode::Authority : EGameplayAbilityActivationMode::NonAuthority);
-	PredictionKey = 0;
-
 	// Grab Components that we care about
 	USkeletalMeshComponent * SkelMeshComponent = InActor->FindComponentByClass<USkeletalMeshComponent>();
 	if (SkelMeshComponent)
@@ -316,30 +310,30 @@ void FGameplayAbilityActorInfo::InitFromActor(AActor *InActor)
 	MovementComponent = InActor->FindComponentByClass<UMovementComponent>();
 }
 
-void FGameplayAbilityActorInfo::GeneratePredictionKey() const
+void FGameplayAbilityActivationInfo::GeneratePredictionKey(UAttributeComponent * Component) const
 {
-	check(AttributeComponent.IsValid());
+	check(Component);
 	check(PredictionKey == 0);
 	check(ActivationMode == EGameplayAbilityActivationMode::NonAuthority);
 
-	PredictionKey = AttributeComponent->GetNextPredictionKey();
+	PredictionKey = Component->GetNextPredictionKey();
 	ActivationMode = EGameplayAbilityActivationMode::Predicting;
 }
 
-void FGameplayAbilityActorInfo::SetPredictionKey(int32 InPredictionKey)
+void FGameplayAbilityActivationInfo::SetPredictionKey(int32 InPredictionKey)
 {
 	PredictionKey = InPredictionKey;
 }
 
-void FGameplayAbilityActorInfo::SetActivationConfirmed()
+void FGameplayAbilityActivationInfo::SetActivationConfirmed()
 {
-	check(ActivationMode == EGameplayAbilityActivationMode::NonAuthority);
+	//check(ActivationMode == EGameplayAbilityActivationMode::NonAuthority);
 	ActivationMode = EGameplayAbilityActivationMode::Confirmed;
 }
 
 //----------------------------------------------------------------------
 
-void FGameplayAbilityTargetData::ApplyGameplayEffect(UGameplayEffect *GameplayEffect, const FGameplayAbilityActorInfo InstigatorInfo)
+void FGameplayAbilityTargetData::ApplyGameplayEffect(UGameplayEffect* GameplayEffect, const FGameplayAbilityActorInfo InstigatorInfo)
 {
 	// Improve relationship between InstigatorContext and 
 	
diff --git a/Engine/Source/Runtime/SkillSystem/Private/Abilities/GameplayAbility_CharacterJump.cpp b/Engine/Source/Runtime/SkillSystem/Private/Abilities/GameplayAbility_CharacterJump.cpp
--- a/Engine/Source/Runtime/SkillSystem/Private/Abilities/GameplayAbility_CharacterJump.cpp
+++ b/Engine/Source/Runtime/SkillSystem/Private/Abilities/GameplayAbility_CharacterJump.cpp
@@ -15,41 +15,46 @@ UGameplayAbility_CharacterJump::UGameplayAbility_CharacterJump(const class FPost
 	NetExecutionPolicy = EGameplayAbilityNetExecutionPolicy::Predictive;
 }
 
-void UGameplayAbility_CharacterJump::ActivateAbility(const FGameplayAbilityActorInfo ActorInfo)
+void UGameplayAbility_CharacterJump::ActivateAbility(const FGameplayAbilityActorInfo* ActorInfo, const FGameplayAbilityActivationInfo ActivationInfo)
 {
-	if (ActorInfo.ActivationMode == EGameplayAbilityActivationMode::Authority ||
-		ActorInfo.ActivationMode == EGameplayAbilityActivationMode::Predicting)
+	if (ActivationInfo.ActivationMode == EGameplayAbilityActivationMode::Authority ||
+		ActivationInfo.ActivationMode == EGameplayAbilityActivationMode::Predicting)
 	{
-		if (!CommitAbility(ActorInfo))
+		if (!CommitAbility(ActorInfo, ActivationInfo))
 		{
 			return;
 		}
 
-		SKILL_LOG(Warning, TEXT("ActivateAbility Jump. %s. ActivationMode: %d"), *GetPathName(), (int32)ActorInfo.ActivationMode);
-
-		ACharacter * Character = CastChecked<ACharacter>(ActorInfo.Actor.Get());
+		ACharacter * Character = CastChecked<ACharacter>(ActorInfo->Actor.Get());
 		Character->Jump();
 	}
 }
 
-void UGameplayAbility_CharacterJump::PredictiveActivateAbility(const FGameplayAbilityActorInfo ActorInfo)
+void UGameplayAbility_CharacterJump::PredictiveActivateAbility(const FGameplayAbilityActorInfo* ActorInfo, const FGameplayAbilityActivationInfo ActivationInfo)
 {
-	CallActivateAbility(ActorInfo);
+	CallActivateAbility(ActorInfo, ActivationInfo);
 }
 
-void UGameplayAbility_CharacterJump::InputReleased(int32 InputID, const FGameplayAbilityActorInfo ActorInfo)
+void UGameplayAbility_CharacterJump::InputReleased(int32 InputID, const FGameplayAbilityActorInfo* ActorInfo)
 {
-	CancelAbility(ActorInfo);
+	// This is the 'top' of the ability execution code path. We don't have an ActivationInfo yet.
+	// Normally for InputPressed, we do a bunch of stuff (for clients): predict, tell the server, wait for confirm, rollback, etc.
+	// For jump, we are just going to call CancelAbility with an authority ActivationInfo. This is all jump needs, but more complex
+	// abilities may need more work.
+
+	FGameplayAbilityActivationInfo	ActivationInfo(EGameplayAbilityActivationMode::Authority, 0);
+
+	CancelAbility(ActorInfo, ActivationInfo);
 }
 
-bool UGameplayAbility_CharacterJump::CanActivateAbility(const FGameplayAbilityActorInfo ActorInfo) const
+bool UGameplayAbility_CharacterJump::CanActivateAbility(const FGameplayAbilityActorInfo* ActorInfo) const
 {
 	if (!Super::CanActivateAbility(ActorInfo))
 	{
 		return false;
 	}
 
-	const UCharacterMovementComponent *CharMovement = CastChecked<UCharacterMovementComponent>(ActorInfo.MovementComponent.Get());
+	const UCharacterMovementComponent *CharMovement = CastChecked<UCharacterMovementComponent>(ActorInfo->MovementComponent.Get());
 	return (CharMovement->MovementMode == EMovementMode::MOVE_Walking);
 }
 
@@ -59,14 +64,12 @@ bool UGameplayAbility_CharacterJump::CanActivateAbility(const FGameplayAbilityAc
  *	Montage that *it* played was still playing, and if so, to cancel it. If this is something we need to support, we may need some
  *	light weight data structure to represent 'non intanced abilities in action' with a way to cancel/end them.
  */
-void UGameplayAbility_CharacterJump::CancelAbility(const FGameplayAbilityActorInfo ActorInfo)
+void UGameplayAbility_CharacterJump::CancelAbility(const FGameplayAbilityActorInfo* ActorInfo, const FGameplayAbilityActivationInfo ActivationInfo)
 {
-	if (ActorInfo.ActivationMode == EGameplayAbilityActivationMode::Authority ||
-		ActorInfo.ActivationMode == EGameplayAbilityActivationMode::Predicting)
+	if (ActivationInfo.ActivationMode == EGameplayAbilityActivationMode::Authority ||
+		ActivationInfo.ActivationMode == EGameplayAbilityActivationMode::Predicting)
 	{
-		SKILL_LOG(Warning, TEXT("CancelAbility Jump. %s. ActivationMode: %d"), *GetPathName(), (int32)ActorInfo.ActivationMode);
-
-		ACharacter * Character = CastChecked<ACharacter>(ActorInfo.Actor.Get());
+		ACharacter * Character = CastChecked<ACharacter>(ActorInfo->Actor.Get());
 		Character->StopJumping();
 	}
 }
\ No newline at end of file
diff --git a/Engine/Source/Runtime/SkillSystem/Private/Abilities/GameplayAbility_Instanced.cpp b/Engine/Source/Runtime/SkillSystem/Private/Abilities/GameplayAbility_Instanced.cpp
--- a/Engine/Source/Runtime/SkillSystem/Private/Abilities/GameplayAbility_Instanced.cpp
+++ b/Engine/Source/Runtime/SkillSystem/Private/Abilities/GameplayAbility_Instanced.cpp
@@ -17,8 +17,10 @@ UGameplayAbility_Instanced::UGameplayAbility_Instanced(const class FPostConstruc
 	ReplicationPolicy = EGameplayAbilityReplicationPolicy::ReplicateOwner;
 }
 
-bool UGameplayAbility_Instanced::CanActivateAbility(const FGameplayAbilityActorInfo ActorInfo) const
+bool UGameplayAbility_Instanced::CanActivateAbility(const FGameplayAbilityActorInfo* ActorInfo) const
 {
+	CurrentActorInfo = ActorInfo;
+
 	if (!Super::CanActivateAbility(ActorInfo))
 	{
 		return false;
@@ -30,7 +32,7 @@ bool UGameplayAbility_Instanced::CanActivateAbility(const FGameplayAbilityActorI
 
 	if (CanActivateFunction
 		&& CanActivateFunction->GetOuter()->IsA(UBlueprintGeneratedClass::StaticClass())
-		&& K2_CanActivateAbility(ActorInfo) == false)
+		&& K2_CanActivateAbility() == false)
 	{
 		SKILL_LOG(Log, TEXT("CanActivateAbility %s failed, blueprint refused"), *GetName());
 		return false;
@@ -39,9 +41,12 @@ bool UGameplayAbility_Instanced::CanActivateAbility(const FGameplayAbilityActorI
 	return true;
 }
 
-void UGameplayAbility_Instanced::CommitExecute(const FGameplayAbilityActorInfo ActorInfo)
+void UGameplayAbility_Instanced::CommitExecute(const FGameplayAbilityActorInfo* ActorInfo, const FGameplayAbilityActivationInfo ActivationInfo)
 {
-	Super::CommitExecute(ActorInfo);
+	CurrentActorInfo = ActorInfo;
+	CurrentActivationInfo = ActivationInfo;
+
+	Super::CommitExecute(ActorInfo, ActivationInfo);
 
 	check(!HasAnyFlags(RF_ClassDefaultObject));
 
@@ -51,25 +56,33 @@ void UGameplayAbility_Instanced::CommitExecute(const FGameplayAbilityActorInfo A
 
 	if (CanActivateFunction	&& CanActivateFunction->GetOuter()->IsA(UBlueprintGeneratedClass::StaticClass()))
 	{
-		K2_CommitExecute(ActorInfo);
+		K2_CommitExecute();
 	}
 }
 
-void UGameplayAbility_Instanced::ActivateAbility(FGameplayAbilityActorInfo ActorInfo)
+void UGameplayAbility_Instanced::ActivateAbility(const FGameplayAbilityActorInfo* ActorInfo, const FGameplayAbilityActivationInfo ActivationInfo)
 {
+	CurrentActorInfo = ActorInfo;
+	CurrentActivationInfo = ActivationInfo;
+
 	// Call into blueprint (fixme: cache this)
 	static FName FuncName = FName(TEXT("K2_ActivateAbility"));
 	UFunction* CanActivateFunction = GetClass()->FindFunctionByName(FuncName);
 
 	if (CanActivateFunction	&& CanActivateFunction->GetOuter()->IsA(UBlueprintGeneratedClass::StaticClass()))
 	{
-		K2_ActivateAbility(ActorInfo);
+		K2_ActivateAbility();
 	}
 }
 
-void UGameplayAbility_Instanced::PredictiveActivateAbility(const FGameplayAbilityActorInfo ActorInfo)
+void UGameplayAbility_Instanced::PredictiveActivateAbility(const FGameplayAbilityActorInfo* ActorInfo, const FGameplayAbilityActivationInfo ActivationInfo)
 {
-	Super::PredictiveActivateAbility(ActorInfo);
+	check(ActivationInfo.ActivationMode == EGameplayAbilityActivationMode::Predicting);
+
+	CurrentActorInfo = ActorInfo;
+	CurrentActivationInfo = ActivationInfo;
+
+	Super::PredictiveActivateAbility(ActorInfo, ActivationInfo);
 
 	// FIXME
 	if(!HasAnyFlags(RF_ClassDefaultObject))
@@ -80,14 +93,68 @@ void UGameplayAbility_Instanced::PredictiveActivateAbility(const FGameplayAbilit
 
 		if (PredictiveActivateFunction&& PredictiveActivateFunction->GetOuter()->IsA(UBlueprintGeneratedClass::StaticClass()))
 		{
-			K2_PredictiveActivateAbility(ActorInfo);
+			K2_PredictiveActivateAbility();
 		}
 	}
 }
 
+void UGameplayAbility_Instanced::EndAbility(const FGameplayAbilityActorInfo* ActorInfo)
+{
+	CurrentActorInfo = ActorInfo;
+
+	Super::EndAbility(ActorInfo);
+
+	if (InstancedPerExecution)
+	{
+		MarkPendingKill();
+	}
+
+	// Remove from owning attributecomponent?
+	// generic way of releasing all callbacks!
+}
+
+FGameplayAbilityActorInfo UGameplayAbility_Instanced::GetActorInfo()
+{
+	check(CurrentActorInfo);
+	return *CurrentActorInfo;
+}
+
+/**
+* Attempts to commit the ability (spend resources, etc). This our last chance to fail.
+*	-Child classes that override ActivateAbility must call this themselves!
+*/
+
+
+bool UGameplayAbility_Instanced::K2_CommitAbility()
+{
+	check(CurrentActorInfo);
+	return CommitAbility(CurrentActorInfo, CurrentActivationInfo);
+}
+
+void UGameplayAbility_Instanced::K2_EndAbility()
+{
+	check(CurrentActorInfo);
+	EndAbility(CurrentActorInfo);
+}
+
+void UGameplayAbility_Instanced::ClientActivateAbilitySucceed_Internal(int32 PredictionKey)
+{
+	check(CurrentActorInfo);
+	CurrentActivationInfo.SetActivationConfirmed();
+
+	CallActivateAbility(CurrentActorInfo, CurrentActivationInfo);
+}
+
+// --------------------------------------------------------------------------------------------------------------------------------------------------------
+//
+//	Replication boilerplate
+//
+// --------------------------------------------------------------------------------------------------------------------------------------------------------
+
+
 int32 UGameplayAbility_Instanced::GetFunctionCallspace(UFunction* Function, void* Parameters, FFrame* Stack)
 {
-	if(HasAnyFlags(RF_ClassDefaultObject))
+	if (HasAnyFlags(RF_ClassDefaultObject))
 	{
 		return FunctionCallspace::Local;
 	}
@@ -111,15 +178,6 @@ bool UGameplayAbility_Instanced::CallRemoteFunction(UFunction* Function, void* P
 	return false;
 }
 
-void UGameplayAbility_Instanced::EndAbility(const FGameplayAbilityActorInfo ActorInfo)
-{
-	Super::EndAbility(ActorInfo);
 
-	if (InstancedPerExecution)
-	{
-		MarkPendingKill();
-	}
 
-	// Remove from owning attributecomponent?
-	// generic way of releasing all callbacks!
-}
\ No newline at end of file
+
diff --git a/Engine/Source/Runtime/SkillSystem/Private/Abilities/GameplayAbility_Montage.cpp b/Engine/Source/Runtime/SkillSystem/Private/Abilities/GameplayAbility_Montage.cpp
--- a/Engine/Source/Runtime/SkillSystem/Private/Abilities/GameplayAbility_Montage.cpp
+++ b/Engine/Source/Runtime/SkillSystem/Private/Abilities/GameplayAbility_Montage.cpp
@@ -18,44 +18,44 @@ UGameplayAbility_Montage::UGameplayAbility_Montage(const class FPostConstructIni
 	NetExecutionPolicy = EGameplayAbilityNetExecutionPolicy::Server;
 }
 
-void UGameplayAbility_Montage::ActivateAbility(const FGameplayAbilityActorInfo OwnerInfo)
+void UGameplayAbility_Montage::ActivateAbility(const FGameplayAbilityActorInfo* ActorInfo, const FGameplayAbilityActivationInfo ActivationInfo)
 {
-	if (!CommitAbility(OwnerInfo))
+	if (!CommitAbility(ActorInfo, ActivationInfo))
 	{
 		return;
 	}
 
-	if (MontageToPlay != NULL && OwnerInfo.AnimInstance != NULL && OwnerInfo.AnimInstance->GetActiveMontageInstance() == NULL)
+	if (MontageToPlay != NULL && ActorInfo->AnimInstance != NULL && ActorInfo->AnimInstance->GetActiveMontageInstance() == NULL)
 	{
 		TArray<FActiveGameplayEffectHandle>	AppliedEffects;
 
 		// Apply GameplayEffects
 		for (const UGameplayEffect * Effect : GameplayEffectsWhileAnimating)
 		{
-			FActiveGameplayEffectHandle Handle = OwnerInfo.AttributeComponent->ApplyGameplayEffectToSelf(Effect, 1.f, OwnerInfo.Actor.Get());
+			FActiveGameplayEffectHandle Handle = ActorInfo->AttributeComponent->ApplyGameplayEffectToSelf(Effect, 1.f, ActorInfo->Actor.Get());
 			if (Handle.IsValid())
 			{
 				AppliedEffects.Add(Handle);
 			}
 		}
 
-		float const Duration = OwnerInfo.AnimInstance->Montage_Play(MontageToPlay, PlayRate);
+		float const Duration = ActorInfo->AnimInstance->Montage_Play(MontageToPlay, PlayRate);
 
 		FOnMontageEnded EndDelegate;
-		EndDelegate.BindUObject(this, &UGameplayAbility_Montage::OnMontageEnded, OwnerInfo, AppliedEffects);
-		OwnerInfo.AnimInstance->Montage_SetEndDelegate(EndDelegate);
+		EndDelegate.BindUObject(this, &UGameplayAbility_Montage::OnMontageEnded, ActorInfo, AppliedEffects);
+		ActorInfo->AnimInstance->Montage_SetEndDelegate(EndDelegate);
 
 		if (SectionName != NAME_None)
 		{
-			OwnerInfo.AnimInstance->Montage_JumpToSection(SectionName);
+			ActorInfo->AnimInstance->Montage_JumpToSection(SectionName);
 		}
 	}
 }
 
-void UGameplayAbility_Montage::OnMontageEnded(UAnimMontage* Montage, bool bInterrupted, const FGameplayAbilityActorInfo OwnerInfo, TArray<FActiveGameplayEffectHandle> AppliedEffects)
+void UGameplayAbility_Montage::OnMontageEnded(UAnimMontage* Montage, bool bInterrupted, const FGameplayAbilityActorInfo* ActorInfo, TArray<FActiveGameplayEffectHandle> AppliedEffects)
 {
 	// Remove any GameplayEffects that we applied
-	UAttributeComponent *OwnerComponent = OwnerInfo.AttributeComponent.Get();
+	UAttributeComponent *OwnerComponent = ActorInfo->AttributeComponent.Get();
 	if (OwnerComponent)
 	{
 		for (FActiveGameplayEffectHandle Handle : AppliedEffects)
diff --git a/Engine/Source/Runtime/SkillSystem/Private/AttributeComponent.cpp b/Engine/Source/Runtime/SkillSystem/Private/AttributeComponent.cpp
--- a/Engine/Source/Runtime/SkillSystem/Private/AttributeComponent.cpp
+++ b/Engine/Source/Runtime/SkillSystem/Private/AttributeComponent.cpp
@@ -26,8 +26,6 @@ static FAutoConsoleVariableRef CVarDebugGameplayCues(
 	ECVF_Default
 	);
 
-#pragma optimize( "", off )
-
 /** Enable to log out all render state create, destroy and updatetransform events */
 #define LOG_RENDER_STATE 0
 
diff --git a/Engine/Source/Runtime/SkillSystem/Private/AttributeComponent_Abilities.cpp b/Engine/Source/Runtime/SkillSystem/Private/AttributeComponent_Abilities.cpp
--- a/Engine/Source/Runtime/SkillSystem/Private/AttributeComponent_Abilities.cpp
+++ b/Engine/Source/Runtime/SkillSystem/Private/AttributeComponent_Abilities.cpp
@@ -18,6 +18,8 @@
 void UAttributeComponent::InitializeComponent()
 {
 	Super::InitializeComponent();
+
+	InitAbilityActorInfo();
 	
 	// Look for DSO AttributeSets
 	AActor *Owner = GetOwner();
@@ -38,6 +40,20 @@ void UAttributeComponent::InitializeComponent()
 	}
 }
 
+void UAttributeComponent::InitAbilityActorInfo()
+{
+	if (!AbilityActorInfo.IsValid())
+	{
+		// Alloc (and init) a new actor info
+		AbilityActorInfo = TSharedPtr<FGameplayAbilityActorInfo>(USkillSystemGlobals::Get().AllocAbilityActorInfo(GetOwner()));
+	}
+	else
+	{
+		// We already have a valid actor info, just reinit it.
+		AbilityActorInfo->InitFromActor(GetOwner());
+	}
+}
+
 UGameplayAbility * UAttributeComponent::CreateNewInstanceOfAbility(UGameplayAbility *Ability)
 {
 	check(Ability);
@@ -64,24 +80,24 @@ UGameplayAbility * UAttributeComponent::CreateNewInstanceOfAbility(UGameplayAbil
 
 bool UAttributeComponent::ActivateAbility(int32 AbilityIdx)
 {
+	check(AbilityActorInfo.IsValid());
+
 	AActor * OwnerActor = GetOwner();
 	check(OwnerActor);
 
-	FGameplayAbilityActorInfo	ActorInfo;
-	ActorInfo.InitFromActor(OwnerActor);
-
 	if (ActivatableAbilities.IsValidIndex(AbilityIdx))
 	{
-		// Fixme
-		ActivatableAbilities[AbilityIdx]->InputPressed(0, ActorInfo);
+		ActivatableAbilities[AbilityIdx]->InputPressed(0, AbilityActorInfo.Get());
 	}
 
 	return false;
 }
 
-void UAttributeComponent::CancelAbilitiesWithTags(const FGameplayTagContainer Tags, const FGameplayAbilityActorInfo & ActorInfo, UGameplayAbility * Ignore)
+void UAttributeComponent::CancelAbilitiesWithTags(const FGameplayTagContainer Tags, const FGameplayAbilityActorInfo* ActorInfo, const FGameplayAbilityActivationInfo ActivationInfo, UGameplayAbility* Ignore)
 {
 	/**
+	 *	FIXME
+	 *
 	 *	Right now we are cancelling all activatable abilites that match Tags. This includes abilities that might not have been activated in the first place!
 	 *	For instanced-per-actor abilities this is fine. They could check if they were activated/still activating.
 	 *	For non-instanced abilities it is ambigous. We have no way to know 'how many' non-instanced abilities are in flight.
@@ -90,14 +106,14 @@ void UAttributeComponent::CancelAbilitiesWithTags(const FGameplayTagContainer Ta
 
 	struct local
 	{
-		static void CancelAbilitiesWithTags(const FGameplayTagContainer Tags, TArray<UGameplayAbility*> Abilities, const FGameplayAbilityActorInfo ActorInfo, UGameplayAbility * Ignore)
+		static void CancelAbilitiesWithTags(const FGameplayTagContainer Tags, TArray<UGameplayAbility*> Abilities, const FGameplayAbilityActorInfo* ActorInfo, const FGameplayAbilityActivationInfo ActivationInfo, UGameplayAbility * Ignore)
 		{
 			for (int32 idx=0; idx < Abilities.Num(); ++idx)
 			{
 				UGameplayAbility *Ability = Abilities[idx];
 				if (Ability && (Ability != Ignore) && Ability->AbilityTags.MatchesAny(Tags, false))
 				{
-					Ability->CancelAbility(ActorInfo);
+					Ability->CancelAbility(ActorInfo, ActivationInfo);
 					if (!Ability->HasAnyFlags(RF_ClassDefaultObject) && Ability->GetInstancingPolicy() == EGameplayAbilityInstancingPolicy::InstancedPerExecution)
 					{
 						Ability->MarkPendingKill();
@@ -112,7 +128,7 @@ void UAttributeComponent::CancelAbilitiesWithTags(const FGameplayTagContainer Ta
 	//local::CancelAbilitiesWithTags(Tags, ReplicatedInstancedAbilities, ActorInfo, Ignore);
 	//local::CancelAbilitiesWithTags(Tags, NonReplicatedInstancedAbilities, ActorInfo, Ignore);
 
-	local::CancelAbilitiesWithTags(Tags, ActivatableAbilities, ActorInfo, Ignore);
+	local::CancelAbilitiesWithTags(Tags, ActivatableAbilities, ActorInfo, ActivationInfo, Ignore);
 }
 
 #if !(UE_BUILD_SHIPPING || UE_BUILD_TEST)
@@ -125,7 +141,7 @@ static FAutoConsoleVariableRef CVarDenyClientActivation(
 	);
 #endif
 
-void UAttributeComponent::ServerTryActivateAbility_Implementation(class UGameplayAbility * AbilityToActivate, int32 PredictionKey)
+void UAttributeComponent::ServerTryActivateAbility_Implementation(UGameplayAbility * AbilityToActivate, int32 PredictionKey)
 {
 #if !(UE_BUILD_SHIPPING || UE_BUILD_TEST)
 	if (DenyClientActivation > 0)
@@ -137,21 +153,18 @@ void UAttributeComponent::ServerTryActivateAbility_Implementation(class UGamepla
 #endif
 
 	ensure(AbilityToActivate);
+	ensure(AbilityActorInfo.IsValid());
 
 	AActor * OwnerActor = GetOwner();
 	check(OwnerActor);
 
-	FGameplayAbilityActorInfo	ActorInfo;
-	ActorInfo.InitFromActor(OwnerActor);
-	ActorInfo.SetPredictionKey(PredictionKey);	// Set the prediction key in case anyone else needs it. But we are still Authority and not Predicting.
-
 	UGameplayAbility *InstancedAbility = NULL;
 
-	if (AbilityToActivate->TryActivateAbility(ActorInfo, &InstancedAbility))
+	if (AbilityToActivate->TryActivateAbility(AbilityActorInfo.Get(), PredictionKey, &InstancedAbility))
 	{
 		if (InstancedAbility && InstancedAbility->GetReplicationPolicy() != EGameplayAbilityReplicationPolicy::ReplicateNone)
 		{
-			InstancedAbility->ClientActivateAbilitySucceed(OwnerActor);
+			InstancedAbility->ClientActivateAbilitySucceed(PredictionKey);
 		}
 		else
 		{
@@ -173,12 +186,12 @@ void UAttributeComponent::ServerTryActivateAbility_Implementation(class UGamepla
 	}
 }
 
-bool UAttributeComponent::ServerTryActivateAbility_Validate(class UGameplayAbility * AbilityToActivate, int32 PredictionKey)
+bool UAttributeComponent::ServerTryActivateAbility_Validate(UGameplayAbility * AbilityToActivate, int32 PredictionKey)
 {
 	return true;
 }
 
-void UAttributeComponent::ClientActivateAbilityFailed_Implementation(class UGameplayAbility * AbilityToActivate, int32 PredictionKey)
+void UAttributeComponent::ClientActivateAbilityFailed_Implementation(UGameplayAbility * AbilityToActivate, int32 PredictionKey)
 {
 	if (PredictionKey > 0)
 	{
@@ -201,19 +214,17 @@ void UAttributeComponent::ClientActivateAbilityFailed_Implementation(class UGame
 	}
 }
 
-void UAttributeComponent::ClientActivateAbilitySucceed_Implementation(class UGameplayAbility * AbilityToActivate, int32 PredictionKey)
+void UAttributeComponent::ClientActivateAbilitySucceed_Implementation(UGameplayAbility* AbilityToActivate, int32 PredictionKey)
 {
 	ensure(AbilityToActivate);
+	ensure(AbilityActorInfo.IsValid());
 
 	AActor * OwnerActor = GetOwner();
 	check(OwnerActor);
 
-	FGameplayAbilityActorInfo	ActorInfo;
-	ActorInfo.InitFromActor(OwnerActor);
-	ActorInfo.SetPredictionKey(PredictionKey);
-	ActorInfo.SetActivationConfirmed();
+	FGameplayAbilityActivationInfo	ActivationInfo(EGameplayAbilityActivationMode::Confirmed, PredictionKey);
 
-	AbilityToActivate->CallActivateAbility(ActorInfo);
+	AbilityToActivate->CallActivateAbility(AbilityActorInfo.Get(), ActivationInfo);
 }
 
 void UAttributeComponent::MontageBranchPoint_AbilityDecisionStop()
@@ -226,7 +237,7 @@ void UAttributeComponent::MontageBranchPoint_AbilityDecisionStop()
 		FGameplayAbilityActorInfo	ActorInfo;
 		ActorInfo.InitFromActor(OwnerActor);
 
-		AnimatingAbility->MontageBranchPoint_AbilityDecisionStop(ActorInfo);
+		AnimatingAbility->MontageBranchPoint_AbilityDecisionStop(AbilityActorInfo.Get());
 	}
 }
 
@@ -237,10 +248,7 @@ void UAttributeComponent::MontageBranchPoint_AbilityDecisionStart()
 		AActor * OwnerActor = GetOwner();
 		check(OwnerActor);
 
-		FGameplayAbilityActorInfo	ActorInfo;
-		ActorInfo.InitFromActor(OwnerActor);
-
-		AnimatingAbility->MontageBranchPoint_AbilityDecisionStart(ActorInfo);
+		AnimatingAbility->MontageBranchPoint_AbilityDecisionStart(AbilityActorInfo.Get());
 	}
 }
 
diff --git a/Engine/Source/Runtime/SkillSystem/Private/AttributeSet.cpp b/Engine/Source/Runtime/SkillSystem/Private/AttributeSet.cpp
--- a/Engine/Source/Runtime/SkillSystem/Private/AttributeSet.cpp
+++ b/Engine/Source/Runtime/SkillSystem/Private/AttributeSet.cpp
@@ -199,7 +199,7 @@ void FScalableFloat::FinalizeCurveData(const FGlobalCurveDataOverride *GlobalOve
 	}
 
 	// Look at global defaults
-	const UCurveTable * GlobalTable = ISkillSystemModule::Get().GetSkillSystemGlobals().GetGlobalCurveTable();
+	const UCurveTable * GlobalTable = ISkillSystemModule::Get().GetSkillSystemGlobals()->GetGlobalCurveTable();
 	if (GlobalTable)
 	{
 		FinalCurve = GlobalTable->FindCurve(Curve.RowName, ContextString, false);
diff --git a/Engine/Source/Runtime/SkillSystem/Private/GameplayAbilitySet.cpp b/Engine/Source/Runtime/SkillSystem/Private/GameplayAbilitySet.cpp
--- a/Engine/Source/Runtime/SkillSystem/Private/GameplayAbilitySet.cpp
+++ b/Engine/Source/Runtime/SkillSystem/Private/GameplayAbilitySet.cpp
@@ -24,7 +24,7 @@ void InputPressed(int32 AbilityIdx, int32 InputID, TWeakObjectPtr<AActor> ActorO
 			UGameplayAbility * Ability = AttributeComponent->ActivatableAbilities[AbilityIdx];
 			if (Ability)
 			{
-				Ability->InputPressed(InputID, ActorInfo);
+				Ability->InputPressed(InputID, AttributeComponent->AbilityActorInfo.Get());
 			}
 		}
 	}
@@ -43,7 +43,7 @@ void InputReleased(int32 AbilityIdx, int32 InputID, TWeakObjectPtr<AActor> Actor
 			UGameplayAbility * Ability = AttributeComponent->ActivatableAbilities[AbilityIdx];
 			if (Ability)
 			{
-				Ability->InputReleased(InputID, ActorInfo);
+				Ability->InputReleased(InputID, AttributeComponent->AbilityActorInfo.Get());
 			}
 		}
 	}
@@ -64,12 +64,6 @@ UGameplayAbilitySet::UGameplayAbilitySet(const class FPostConstructInitializePro
 
 void UGameplayAbilitySet::InitializeAbilities(UAttributeComponent *AttributeComponent) const
 {
-	AActor * OwnerActor = AttributeComponent->GetOwner();
-	check(OwnerActor);
-
-	FGameplayAbilityActorInfo	ActorInfo;
-	ActorInfo.InitFromActor(OwnerActor);
-
 	for (const FGameplayAbilityBindInfo &BindInfo : this->Abilities)
 	{
 		UGameplayAbility * Ability = BindInfo.GameplayAbilityClass ? BindInfo.GameplayAbilityClass->GetDefaultObject<UGameplayAbility>() : BindInfo.GameplayAbilityInstance;
@@ -90,9 +84,6 @@ void UGameplayAbilitySet::BindInputComponentToAbilities(UInputComponent *InputCo
 	AActor * OwnerActor = AttributeComponent->GetOwner();
 	check(OwnerActor);
 
-	FGameplayAbilityActorInfo	ActorInfo;
-	ActorInfo.InitFromActor(OwnerActor);
-
 	for (int32 idx = 0; idx < this->Abilities.Num(); ++idx)
 	{
 		const FGameplayAbilityBindInfo & BindInfo = this->Abilities[idx];
diff --git a/Engine/Source/Runtime/SkillSystem/Private/GameplayEffect.cpp b/Engine/Source/Runtime/SkillSystem/Private/GameplayEffect.cpp
--- a/Engine/Source/Runtime/SkillSystem/Private/GameplayEffect.cpp
+++ b/Engine/Source/Runtime/SkillSystem/Private/GameplayEffect.cpp
@@ -11,8 +11,6 @@ const float UGameplayEffect::INSTANT_APPLICATION = 0.f;
 const float UGameplayEffect::NO_PERIOD = 0.f;
 const float FGameplayEffectLevelSpec::INVALID_LEVEL = -1.f;
 
-#pragma optimize( "", off )
-
 #if SKILL_SYSTEM_AGGREGATOR_DEBUG
 FAggregator::FAllocationStats FAggregator::AllocationStats;
 #endif
@@ -80,7 +78,7 @@ FGameplayEffectSpec::FGameplayEffectSpec(const UGameplayEffect *InDef, AActor *
 		}
 		else
 		{
-			UDataTable* DataTable = ISkillSystemModule::Get().GetSkillSystemGlobals().GetGlobalAttributeDataTable();
+			UDataTable* DataTable = ISkillSystemModule::Get().GetSkillSystemGlobals()->GetGlobalAttributeDataTable();
 
 			if (Def->Modifiers[0].ModifierType != EGameplayMod::Attribute)
 			{
@@ -126,7 +124,7 @@ FGameplayEffectSpec::FGameplayEffectSpec(const UGameplayEffect *InDef, AActor *
 			}
 		}
 		
-		UDataTable* DataTable = ISkillSystemModule::Get().GetSkillSystemGlobals().GetGlobalAttributeDataTable();
+		UDataTable* DataTable = ISkillSystemModule::Get().GetSkillSystemGlobals()->GetGlobalAttributeDataTable();
 
 		// check for other attributes that may allow stacking and warn that they will be ignored
 		for (int32 Idx = 1; Idx < Def->Modifiers.Num(); ++Idx)
diff --git a/Engine/Source/Runtime/SkillSystem/Private/GameplayEffectTests.cpp b/Engine/Source/Runtime/SkillSystem/Private/GameplayEffectTests.cpp
--- a/Engine/Source/Runtime/SkillSystem/Private/GameplayEffectTests.cpp
+++ b/Engine/Source/Runtime/SkillSystem/Private/GameplayEffectTests.cpp
@@ -2323,13 +2323,13 @@ UCurveTable * SetGlobalCurveTable()
 		check(Value == 5.f);
 	}
 
-	ISkillSystemModule::Get().GetSkillSystemGlobals().AutomationTestOnly_SetGlobalCurveTable(CurveTable);
+	ISkillSystemModule::Get().GetSkillSystemGlobals()->AutomationTestOnly_SetGlobalCurveTable(CurveTable);
 	return CurveTable;
 }
 
 void ClearGlobalCurveTable()
 {
-	ISkillSystemModule::Get().GetSkillSystemGlobals().AutomationTestOnly_SetGlobalCurveTable(NULL);
+	ISkillSystemModule::Get().GetSkillSystemGlobals()->AutomationTestOnly_SetGlobalCurveTable(NULL);
 }
 
 UCurveTable * GetStandardDamageOverrideCurveTable(float Factor)
@@ -2369,13 +2369,13 @@ UDataTable * SetGlobalDataTable()
 		check(!Row->bCanStack);
 	}
 
-	ISkillSystemModule::Get().GetSkillSystemGlobals().AutomationTestOnly_SetGlobalAttributeDataTable(DataTable);
+	ISkillSystemModule::Get().GetSkillSystemGlobals()->AutomationTestOnly_SetGlobalAttributeDataTable(DataTable);
 	return DataTable;
 }
 
 void ClearGlobalDataTable()
 {
-	ISkillSystemModule::Get().GetSkillSystemGlobals().AutomationTestOnly_SetGlobalAttributeDataTable(NULL);
+	ISkillSystemModule::Get().GetSkillSystemGlobals()->AutomationTestOnly_SetGlobalAttributeDataTable(NULL);
 }
 
 bool GameplayEffectsTest_InstantDamage_ScalingExplicit(UWorld *World, FAutomationTestBase * Test)
@@ -5908,8 +5908,8 @@ bool GameplayEffectsTest_ModifyChanceToExecuteOnGE(UWorld *World, FAutomationTes
 
 bool FGameplayEffectsTest::RunTest( const FString& Parameters )
 {
-	UCurveTable *CurveTable = ISkillSystemModule::Get().GetSkillSystemGlobals().GetGlobalCurveTable();
-	UDataTable *DataTable = ISkillSystemModule::Get().GetSkillSystemGlobals().GetGlobalAttributeDataTable();
+	UCurveTable *CurveTable = ISkillSystemModule::Get().GetSkillSystemGlobals()->GetGlobalCurveTable();
+	UDataTable *DataTable = ISkillSystemModule::Get().GetSkillSystemGlobals()->GetGlobalAttributeDataTable();
 
 	UWorld *World = UWorld::CreateWorld(EWorldType::Game, false);
 	FWorldContext &WorldContext = GEngine->CreateNewWorldContext(EWorldType::Game);
@@ -6036,8 +6036,8 @@ bool FGameplayEffectsTest::RunTest( const FString& Parameters )
 	GEngine->DestroyWorldContext(World);
 	World->DestroyWorld(false);
 
-	ISkillSystemModule::Get().GetSkillSystemGlobals().AutomationTestOnly_SetGlobalCurveTable(CurveTable);
-	ISkillSystemModule::Get().GetSkillSystemGlobals().AutomationTestOnly_SetGlobalAttributeDataTable(DataTable);
+	ISkillSystemModule::Get().GetSkillSystemGlobals()->AutomationTestOnly_SetGlobalCurveTable(CurveTable);
+	ISkillSystemModule::Get().GetSkillSystemGlobals()->AutomationTestOnly_SetGlobalAttributeDataTable(DataTable);
 	return true;
 }
 
diff --git a/Engine/Source/Runtime/SkillSystem/Private/SkillSystem.cpp b/Engine/Source/Runtime/SkillSystem/Private/SkillSystem.cpp
--- a/Engine/Source/Runtime/SkillSystem/Private/SkillSystem.cpp
+++ b/Engine/Source/Runtime/SkillSystem/Private/SkillSystem.cpp
@@ -6,4 +6,5 @@ DEFINE_STAT(STAT_GameplayEffectsHasAllTags);
 DEFINE_STAT(STAT_GameplayEffectsHasAnyTag);
 DEFINE_STAT(STAT_GameplayEffectsTick);
 DEFINE_STAT(STAT_GameplayEffectsCanApplyAttributeModifiers);
-DEFINE_STAT(STAT_GameplayEffectsGetActiveEffectsData);
\ No newline at end of file
+DEFINE_STAT(STAT_GameplayEffectsGetActiveEffectsData);
+DEFINE_STAT(STAT_GameplayAbilityGetCooldownTimeRemaining);
\ No newline at end of file
diff --git a/Engine/Source/Runtime/SkillSystem/Private/SkillSystemBlueprintLibrary.cpp b/Engine/Source/Runtime/SkillSystem/Private/SkillSystemBlueprintLibrary.cpp
--- a/Engine/Source/Runtime/SkillSystem/Private/SkillSystemBlueprintLibrary.cpp
+++ b/Engine/Source/Runtime/SkillSystem/Private/SkillSystemBlueprintLibrary.cpp
@@ -109,13 +109,13 @@ UAbilityTask_WaitOverlap* USkillSystemBlueprintLibrary::CreateWaitOverlap(class
 
 UAttributeComponent* USkillSystemBlueprintLibrary::GetAttributeComponent(AActor *Actor)
 {
-	return ISkillSystemModule::Get().GetSkillSystemGlobals().GetAttributeComponentFromActor(Actor);
+	return USkillSystemGlobals::Get().GetAttributeComponentFromActor(Actor);
 }
 
 void USkillSystemBlueprintLibrary::ApplyGameplayEffectToTargetData(FGameplayAbilityTargetDataHandle Target, UGameplayEffect *GameplayEffect, const FGameplayAbilityActorInfo InstigatorInfo)
 {
 	if (Target.Data.IsValid())
 	{
-		Target.Data->ApplyGameplayEffect(GameplayEffect, InstigatorInfo);	
+		Target.Data->ApplyGameplayEffect(GameplayEffect, InstigatorInfo);
 	}
 }
\ No newline at end of file
diff --git a/Engine/Source/Runtime/SkillSystem/Private/SkillSystemGlobals.cpp b/Engine/Source/Runtime/SkillSystem/Private/SkillSystemGlobals.cpp
--- a/Engine/Source/Runtime/SkillSystem/Private/SkillSystemGlobals.cpp
+++ b/Engine/Source/Runtime/SkillSystem/Private/SkillSystemGlobals.cpp
@@ -2,6 +2,7 @@
 
 #include "SkillSystemModulePrivatePCH.h"
 #include "AttributeComponent.h"
+#include "SkillSystemTypes.h"
 
 USkillSystemGlobals::USkillSystemGlobals(const class FPostConstructInitializeProperties& PCIP)
 : Super(PCIP)
@@ -72,4 +73,21 @@ UAttributeComponent * USkillSystemGlobals::GetAttributeComponentFromActor(AActor
 	// Caller should check this
 	ensure(false);
 	return NULL;
-}
\ No newline at end of file
+}
+
+FGameplayAbilityActorInfo * USkillSystemGlobals::AllocAbilityActorInfo(AActor *Actor) const
+{
+	FGameplayAbilityActorInfo * Info = new FGameplayAbilityActorInfo();
+	Info->InitFromActor(Actor);
+	return Info;
+}
+
+/** This is just some syntax sugar to avoid calling gode to have to ISkillSystemModule::Get() */
+USkillSystemGlobals& USkillSystemGlobals::Get()
+{
+	static USkillSystemGlobals * GlobalPtr = ISkillSystemModule::Get().GetSkillSystemGlobals();
+	check(GlobalPtr == ISkillSystemModule::Get().GetSkillSystemGlobals());
+	check(GlobalPtr);
+
+	return *GlobalPtr;
+}
diff --git a/Engine/Source/Runtime/SkillSystem/Private/SkillSystemModule.cpp b/Engine/Source/Runtime/SkillSystem/Private/SkillSystemModule.cpp
--- a/Engine/Source/Runtime/SkillSystem/Private/SkillSystemModule.cpp
+++ b/Engine/Source/Runtime/SkillSystem/Private/SkillSystemModule.cpp
@@ -11,10 +11,10 @@ class FSkillSystemModule : public ISkillSystemModule
 	virtual void ShutdownModule() OVERRIDE;
 	// End IModuleInterface
 
-	virtual USkillSystemGlobals& GetSkillSystemGlobals()
+	virtual USkillSystemGlobals* GetSkillSystemGlobals()
 	{
 		check(SkillSystemGlobals);
-		return *SkillSystemGlobals;
+		return SkillSystemGlobals;
 	}
 
 	USkillSystemGlobals *SkillSystemGlobals;
@@ -26,7 +26,7 @@ private:
 IMPLEMENT_MODULE( FSkillSystemModule, SkillSystem )
 
 void FSkillSystemModule::StartupModule()
-{
+{	
 	// This code will execute after your module is loaded into memory (but after global variables are initialized, of course.)
 	SkillSystemGlobals = ConstructObject<USkillSystemGlobals>(USkillSystemGlobals::StaticClass(), GetTransientPackage(), NAME_None, RF_RootSet);
 }
diff --git a/Engine/Source/Runtime/SkillSystem/Public/SkillSystem.h b/Engine/Source/Runtime/SkillSystem/Public/SkillSystem.h
--- a/Engine/Source/Runtime/SkillSystem/Public/SkillSystem.h
+++ b/Engine/Source/Runtime/SkillSystem/Public/SkillSystem.h
@@ -13,4 +13,5 @@ DECLARE_CYCLE_STAT_EXTERN(TEXT("GameplayEffectsHasAllTagsTime"), STAT_GameplayEf
 DECLARE_CYCLE_STAT_EXTERN(TEXT("GameplayEffectsHasAnyTagTime"), STAT_GameplayEffectsHasAnyTag, STATGROUP_SkillSystem, );
 DECLARE_CYCLE_STAT_EXTERN(TEXT("GameplayEffectsTick"), STAT_GameplayEffectsTick, STATGROUP_SkillSystem, );
 DECLARE_CYCLE_STAT_EXTERN(TEXT("CanApplyAttributeModifiers"), STAT_GameplayEffectsCanApplyAttributeModifiers, STATGROUP_SkillSystem, );
-DECLARE_CYCLE_STAT_EXTERN(TEXT("GetActiveEffectsData"), STAT_GameplayEffectsGetActiveEffectsData, STATGROUP_SkillSystem, );
\ No newline at end of file
+DECLARE_CYCLE_STAT_EXTERN(TEXT("GetActiveEffectsData"), STAT_GameplayEffectsGetActiveEffectsData, STATGROUP_SkillSystem, );
+DECLARE_CYCLE_STAT_EXTERN(TEXT("GetCooldownTimeRemaining"), STAT_GameplayAbilityGetCooldownTimeRemaining, STATGROUP_SkillSystem, );
\ No newline at end of file
diff --git a/Engine/Source/Runtime/SkillSystem/Public/SkillSystemModule.h b/Engine/Source/Runtime/SkillSystem/Public/SkillSystemModule.h
--- a/Engine/Source/Runtime/SkillSystem/Public/SkillSystemModule.h
+++ b/Engine/Source/Runtime/SkillSystem/Public/SkillSystemModule.h
@@ -35,6 +35,6 @@ public:
 		return FModuleManager::Get().IsModuleLoaded( "SkillSystem" );
 	}
 
-	virtual USkillSystemGlobals& GetSkillSystemGlobals() = 0;
+	virtual USkillSystemGlobals* GetSkillSystemGlobals() = 0;
 };
 
diff --git a/Engine/Source/Runtime/SkillSystem/Public/SkillSystemTypes.h b/Engine/Source/Runtime/SkillSystem/Public/SkillSystemTypes.h
--- a/Engine/Source/Runtime/SkillSystem/Public/SkillSystemTypes.h
+++ b/Engine/Source/Runtime/SkillSystem/Public/SkillSystemTypes.h
@@ -20,25 +20,21 @@ namespace EGameplayAbilityActivationMode
 	};
 }
 
+/**
+ *	FGameplayAbilityActorInfo 
+ *	
+ *	Cached data associated with an Actor using an Ability.
+ *		-Initialized from an AActor* in InitFromActor
+ *		-Abilities use this to know what to actor upon. E.g., instead of being coupled to a specific actor class.
+ *		-These are generally passed around as pointers to support polymorphism.
+ *		-Projects can override USkillSystemGlobals::AllocAbilityActorInfo to override the default struct type that is created. 
+ *			
+ */
 USTRUCT(BlueprintType)
 struct SKILLSYSTEM_API FGameplayAbilityActorInfo
 {
 	GENERATED_USTRUCT_BODY()
 
-	FGameplayAbilityActorInfo()
-		: ActivationMode(EGameplayAbilityActivationMode::Authority)
-		, PredictionKey(0)
-	{
-
-	}
-
-	FGameplayAbilityActorInfo(EGameplayAbilityActivationMode::Type InType)
-		: ActivationMode(InType)
-		, PredictionKey(0)
-	{
-	}
-
-
 	UPROPERTY(BlueprintReadOnly, Category = "ActorInfo")
 	TWeakObjectPtr<AActor>	Actor;
 
@@ -51,19 +47,57 @@ struct SKILLSYSTEM_API FGameplayAbilityActorInfo
 	UPROPERTY(BlueprintReadOnly, Category = "ActorInfo")
 	TWeakObjectPtr<UMovementComponent>	MovementComponent;
 
-	UPROPERTY(BlueprintReadOnly, Category = "ActorInfo")
-	mutable TEnumAsByte<EGameplayAbilityActivationMode::Type>	ActivationMode;
+	virtual void InitFromActor(AActor *Actor);
+};
 
-	UPROPERTY()
-	mutable int32 PredictionKey;
+/**
+ *	FGameplayAbilityActivationInfo
+ *	
+ *	Data tied to a specific activation of an ability. 
+ *		-Tell us whether we are the authority, if we are predicting, confirmed, etc.
+ *		-Holds PredictionKey
+ *		-Generally not meant to be subclassed in projects.
+ *		-Passed around by value since the struct is small.
+ */
 
-	void InitFromActor(AActor *Actor);
+USTRUCT(BlueprintType)
+struct SKILLSYSTEM_API FGameplayAbilityActivationInfo
+{
+	GENERATED_USTRUCT_BODY()
+
+	FGameplayAbilityActivationInfo()
+		: PredictionKey(0)
+		, ActivationMode(EGameplayAbilityActivationMode::Authority)
+	{
+
+	}
+
+	FGameplayAbilityActivationInfo(AActor * InActor, int32 InPredictionKey)
+		: PredictionKey(InPredictionKey)
+	{
+		// On Init, we are either Authority or NonAuthority. We haven't been given a PredictionKey and we haven't been confirmed.
+		// NonAuthority essentially means 'I'm not sure what how I'm going to do this yet'.
+		ActivationMode = (InActor->Role == ROLE_Authority ? EGameplayAbilityActivationMode::Authority : EGameplayAbilityActivationMode::NonAuthority);
+	}
+
+	FGameplayAbilityActivationInfo(EGameplayAbilityActivationMode::Type InType, int32 InPredictionKey)
+		: ActivationMode(InType)
+		, PredictionKey(InPredictionKey)
+	{
+	}
 
-	void GeneratePredictionKey() const;
+
+	void GeneratePredictionKey(UAttributeComponent * Component) const;
 
 	void SetPredictionKey(int32 InPredictionKey);
 
 	void SetActivationConfirmed();
+
+	UPROPERTY(BlueprintReadOnly, Category = "ActorInfo")
+	mutable TEnumAsByte<EGameplayAbilityActivationMode::Type>	ActivationMode;
+
+	UPROPERTY()
+	mutable int32 PredictionKey;
 };
 
 
