Cleanup gameplay effect stacking code.

Moved validation checks out of GameplayEffectSpec constructor and into GameplayEffect::PostLoad.
Stacking attribute is now set in the GameplayEffect.

diff --git a/Engine/Source/Runtime/SkillSystem/Classes/GameplayEffect.h b/Engine/Source/Runtime/SkillSystem/Classes/GameplayEffect.h
--- a/Engine/Source/Runtime/SkillSystem/Classes/GameplayEffect.h
+++ b/Engine/Source/Runtime/SkillSystem/Classes/GameplayEffect.h
@@ -413,10 +413,7 @@ public:
 	/** Get the "clear tags" for the effect */
 	virtual void GetClearGameplayTags(FGameplayTagContainer& TagContainer) const;
 
-	void ValidateGameplayEffect()
-	{
-		// Instant modifiers cannot have incoming/outgoing combat effect modifiers
-	}
+	void ValidateGameplayEffect();
 
 	/** Container of "owned" gameplay tags that are applied to actor with this effect applied to them */
 	UPROPERTY(EditDefaultsOnly, BlueprintReadOnly, Category=Tags)
@@ -432,6 +429,8 @@ public:
 		return false;
 	}
 
+	virtual void PostLoad() OVERRIDE;
+
 	// -----------------------------------------------
 	
 	UPROPERTY(EditDefaultsOnly, BlueprintReadOnly, Category = Advanced)
@@ -451,8 +450,14 @@ public:
 	UPROPERTY(EditDefaultsOnly, BlueprintReadOnly, Category = Stacking)
 	TEnumAsByte<EGameplayEffectStackingPolicy::Type>	StackingPolicy;
 
+	UPROPERTY(EditDefaultsOnly, BlueprintReadOnly, Category = Stacking)
+	FName StackedAttribName;
+
 	UPROPERTY(EditDefaultsOnly, Category = Stacking)
 	TSubclassOf<class UGameplayEffectStackingExtension> StackingExtension;
+
+protected:
+	void ValidateStacking();
 };
 
 /**
@@ -1138,7 +1143,7 @@ struct FGameplayEffectSpec
 		: ModifierLevel( TSharedPtr< FGameplayEffectLevelSpec >( new FGameplayEffectLevelSpec() ) )
 		, Duration(new FAggregator(FGameplayModifierEvaluatedData(0.f, NULL, FActiveGameplayEffectHandle()), SKILL_AGG_DEBUG(TEXT("Uninitialized Duration"))))
 		, Period(new FAggregator(FGameplayModifierEvaluatedData(0.f, NULL, FActiveGameplayEffectHandle()), SKILL_AGG_DEBUG(TEXT("Uninitialized Period"))))
-		, StackedAttribName(NAME_None)
+		, bTopOfStack(false)
 	{
 		// If we initialize a GameplayEffectSpec with no level object passed in.
 	}
@@ -1177,9 +1182,9 @@ struct FGameplayEffectSpec
 	UPROPERTY()
 	FAggregatorRef	ChanceToExecuteOnGameplayEffect;
 
-	// How this combines with other gameplay effects
-	EGameplayEffectStackingPolicy::Type StackingPolicy;
-	FName StackedAttribName;
+	// This should only be true if this is a stacking effect and at the top of its stack
+	UPROPERTY()
+	bool bTopOfStack;
 
 	// The spec needs to own these FModifierSpecs so that other people can keep TSharedPtr to it.
 	// The stuff in this array is OWNED by this spec
diff --git a/Engine/Source/Runtime/SkillSystem/Private/AttributeComponent.cpp b/Engine/Source/Runtime/SkillSystem/Private/AttributeComponent.cpp
--- a/Engine/Source/Runtime/SkillSystem/Private/AttributeComponent.cpp
+++ b/Engine/Source/Runtime/SkillSystem/Private/AttributeComponent.cpp
@@ -273,7 +273,7 @@ FActiveGameplayEffectHandle UAttributeComponent::ApplyGameplayEffectSpecToSelf(c
 		if (Duration != UGameplayEffect::INSTANT_APPLICATION)
 		{
 			// recalculating stacking needs to come before creating the new effect
-			if (Spec.StackingPolicy != EGameplayEffectStackingPolicy::Unlimited)
+			if (Spec.GetStackingType() != EGameplayEffectStackingPolicy::Unlimited)
 			{
 				ActiveGameplayEffects.StacksNeedToRecalculate();
 			}
diff --git a/Engine/Source/Runtime/SkillSystem/Private/GameplayEffect.cpp b/Engine/Source/Runtime/SkillSystem/Private/GameplayEffect.cpp
--- a/Engine/Source/Runtime/SkillSystem/Private/GameplayEffect.cpp
+++ b/Engine/Source/Runtime/SkillSystem/Private/GameplayEffect.cpp
@@ -27,6 +27,8 @@ UGameplayEffect::UGameplayEffect(const class FPostConstructInitializeProperties&
 {
 	ChanceToApplyToTarget.SetValue(1.f);
 	ChanceToExecuteOnGameplayEffect.SetValue(1.f);
+	StackingPolicy = EGameplayEffectStackingPolicy::Unlimited;
+	StackedAttribName = NAME_None;
 }
 
 void UGameplayEffect::GetOwnedGameplayTags(FGameplayTagContainer& TagContainer) const
@@ -44,105 +46,95 @@ bool UGameplayEffect::AreApplicationTagRequirementsSatisfied(const FGameplayTagC
 	return (InstigatorTags.MatchesAll(ApplicationRequiredInstigatorTags, true) && TargetTags.MatchesAll(ApplicationRequiredTargetTags, true));
 }
 
-// --------------------------------------------------------------------------------------------------------------------------------------------------------
-//
-//	FGameplayEffectSpec
-//
-// --------------------------------------------------------------------------------------------------------------------------------------------------------
+void UGameplayEffect::PostLoad()
+{
+	Super::PostLoad();
 
-FGameplayEffectSpec::FGameplayEffectSpec(const UGameplayEffect *InDef, AActor * Instigator, float Level, const FGlobalCurveDataOverride *CurveData)
-: Def(InDef)
-, ModifierLevel(TSharedPtr<FGameplayEffectLevelSpec>(new FGameplayEffectLevelSpec(Level, InDef->LevelInfo, Instigator)))
-, Duration(new FAggregator(InDef->Duration.MakeFinalizedCopy(CurveData), ModifierLevel, SKILL_AGG_DEBUG(TEXT("%s Duration"), *InDef->GetName())))
-, Period(new FAggregator(InDef->Period.MakeFinalizedCopy(CurveData), ModifierLevel, SKILL_AGG_DEBUG(TEXT("%s Period"), *InDef->GetName())))
-, ChanceToApplyToTarget(new FAggregator(InDef->ChanceToApplyToTarget.MakeFinalizedCopy(CurveData), ModifierLevel, SKILL_AGG_DEBUG(TEXT("%s ChanceToApplyToTarget"), *InDef->GetName())))
-, ChanceToExecuteOnGameplayEffect(new FAggregator(InDef->ChanceToExecuteOnGameplayEffect.MakeFinalizedCopy(CurveData), ModifierLevel, SKILL_AGG_DEBUG(TEXT("%s ChanceToExecuteOnGameplayEffect"), *InDef->GetName())))
-, StackingPolicy(EGameplayEffectStackingPolicy::Unlimited)
-, StackedAttribName(NAME_None)
+	ValidateGameplayEffect();
+}
+
+void UGameplayEffect::ValidateGameplayEffect()
 {
-	Duration.Get()->RegisterLevelDependancies();
-	Period.Get()->RegisterLevelDependancies();
-	ChanceToApplyToTarget.Get()->RegisterLevelDependancies();
-	ChanceToExecuteOnGameplayEffect.Get()->RegisterLevelDependancies();
+	ValidateStacking();
 
-	// check the stacking policy for this gameplay effect and warn if it's set up badly
-	// FIXME: this check should be done in the UGameplayEffect PostEdit or PostLoad - not in the spec startup
-	if (Def->StackingPolicy != EGameplayEffectStackingPolicy::Unlimited)
+	// todo: add a check here for instant effects that modify incoming/outgoing GEs
+}
+
+void UGameplayEffect::ValidateStacking()
+{
+	if (StackingPolicy != EGameplayEffectStackingPolicy::Unlimited)
 	{
-		StackingPolicy = Def->StackingPolicy;
+		StackingPolicy = StackingPolicy;
 
-		if (Def->Modifiers.Num() < 1)
+		if (StackedAttribName == NAME_None)
 		{
-			SKILL_LOG(Warning, TEXT("%s has a stacking rule but does not have any modifiers to apply it to."), *Def->GetPathName());
+			SKILL_LOG(Warning, TEXT("%s has a stacking rule but does not have an attribute to apply it to. Removing stacking rule."), *GetPathName());
 			StackingPolicy = EGameplayEffectStackingPolicy::Unlimited;
 		}
-		else
-		{
-			UDataTable* DataTable = ISkillSystemModule::Get().GetSkillSystemGlobals()->GetGlobalAttributeDataTable();
 
-			if (Def->Modifiers[0].ModifierType != EGameplayMod::Attribute)
+		UDataTable* DataTable = ISkillSystemModule::Get().GetSkillSystemGlobals()->GetGlobalAttributeDataTable();
+
+		if (DataTable && DataTable->RowMap.Contains(StackedAttribName))
+		{
+			FAttributeMetaData* Row = (FAttributeMetaData*)DataTable->RowMap[StackedAttribName];
+			if (Row && Row->bCanStack == false)
 			{
-				SKILL_LOG(Warning, TEXT("%s has a stacking rule but its first modifier does not modify an attribute."), *Def->GetPathName());
+				SKILL_LOG(Warning, TEXT("%s has a stacking rule but modifies attribute %s. %s is not allowed to stack. Removing stacking rule."), *GetPathName(), *Modifiers[0].Attribute.GetName(), *Modifiers[0].Attribute.GetName());
 				StackingPolicy = EGameplayEffectStackingPolicy::Unlimited;
 			}
-			else
-			{
-				FName AttributeName(*Def->Modifiers[0].Attribute.GetName());
-				
-				if (DataTable && DataTable->RowMap.Contains(AttributeName))
-				{
-					FAttributeMetaData* Row = (FAttributeMetaData*)DataTable->RowMap[AttributeName];
-					if (Row && Row->bCanStack == false)
-					{
-						SKILL_LOG(Warning, TEXT("%s has a stacking rule but modifies attribute %s. %s is not allowed to stack."), *Def->GetPathName(), *Def->Modifiers[0].Attribute.GetName(), *Def->Modifiers[0].Attribute.GetName());
-						StackingPolicy = EGameplayEffectStackingPolicy::Unlimited;
-					}
-				}
-				else
-				{
-					SKILL_LOG(Warning, TEXT("%s has a stacking rule but modifies attribute %s. %s was not found in the global attribute data table."), *Def->GetPathName(), *Def->Modifiers[0].Attribute.GetName(), *Def->Modifiers[0].Attribute.GetName());
-					StackingPolicy = EGameplayEffectStackingPolicy::Unlimited;
-				}
-			}
+		}
+		else
+		{
+			SKILL_LOG(Warning, TEXT("%s has a stacking rule but modifies attribute %s. %s was not found in the global attribute data table. Removing stacking rule."), *GetPathName(), *Modifiers[0].Attribute.GetName(), *Modifiers[0].Attribute.GetName());
+			StackingPolicy = EGameplayEffectStackingPolicy::Unlimited;
+		}
 
-			// look for other stacking attributes and warn that they will be ignored
-			for (int32 Idx = 1; Idx < Def->Modifiers.Num(); ++Idx)
+		if (Modifiers.Num() == 0)
+		{
+			SKILL_LOG(Warning, TEXT("%s has a stacking rule but does not have any modifiers to apply it to. Removing stacking rule."), *GetPathName());
+			StackingPolicy = EGameplayEffectStackingPolicy::Unlimited;
+		}
+		else
+		{
+			// look for modifiers to attributes that are not the stacking attribute and warn that they may not apply in a predictable fashion
+			for (int32 Idx = 0; Idx < Modifiers.Num(); ++Idx)
 			{
-				if (Def->Modifiers[Idx].ModifierType == EGameplayMod::Attribute)
+				if (Modifiers[Idx].ModifierType == EGameplayMod::Attribute)
 				{
-					FName AttributeName(*Def->Modifiers[0].Attribute.GetName());
+					FName AttributeName(*Modifiers[Idx].Attribute.GetName());
 
-					if (DataTable && DataTable->RowMap.Contains(AttributeName))
+					if (AttributeName != StackedAttribName)
 					{
-						FAttributeMetaData* Row = (FAttributeMetaData*)DataTable->RowMap[AttributeName];
-						if (Row && Row->bCanStack)
-						{
-							SKILL_LOG(Warning, TEXT("%s has a stacking rule and modifies attribute %s but %s is not the modified by the first modifier. Stacking is based on the first modifier."), *Def->GetPathName(), *Def->Modifiers[0].Attribute.GetName(), *Def->Modifiers[0].Attribute.GetName());
-						}
+						SKILL_LOG(Warning, TEXT("%s has a stacking rule and modifies attribute %s but one of the modifiers modifies attribute %s. Changes to %s may not behave in a predictable fashion."), *GetPathName(), *StackedAttribName.ToString(), *AttributeName.ToString(), *AttributeName.ToString());
 					}
 				}
-			}
-		}
-		
-		UDataTable* DataTable = ISkillSystemModule::Get().GetSkillSystemGlobals()->GetGlobalAttributeDataTable();
-
-		// check for other attributes that may allow stacking and warn that they will be ignored
-		for (int32 Idx = 1; Idx < Def->Modifiers.Num(); ++Idx)
-		{
-			if (Def->Modifiers[Idx].ModifierType == EGameplayMod::Attribute)
-			{
-				FName AttributeName(*Def->Modifiers[Idx].Attribute.GetName());
-				if (DataTable && DataTable->RowMap.Contains(AttributeName))
+				else if (Modifiers[Idx].ModifierType == EGameplayMod::ActiveGE || Modifiers[Idx].ModifierType == EGameplayMod::IncomingGE || Modifiers[Idx].ModifierType == EGameplayMod::OutgoingGE)
 				{
-					FAttributeMetaData* Row = (FAttributeMetaData*)DataTable->RowMap[AttributeName];
-					if (Row && Row->bCanStack == true)
-					{
-						SKILL_LOG(Warning, TEXT("%s has a stacking rule and modifies more than one stackable attribute. Stacking will  not apply to attribute %s"), *Def->GetPathName(), *Def->Modifiers[Idx].Attribute.GetName());
-					}
+					SKILL_LOG(Warning, TEXT("%s has a stacking rule and modifies attribute %s but one of its modifiers modifies gameplay effects. Changes to other gameplay effects may not behave in a predictable fashion."), *GetPathName(), *StackedAttribName.ToString());
 				}
 			}
 		}
 	}
+}
+
+// --------------------------------------------------------------------------------------------------------------------------------------------------------
+//
+//	FGameplayEffectSpec
+//
+// --------------------------------------------------------------------------------------------------------------------------------------------------------
+
+FGameplayEffectSpec::FGameplayEffectSpec(const UGameplayEffect *InDef, AActor * Instigator, float Level, const FGlobalCurveDataOverride *CurveData)
+: Def(InDef)
+, ModifierLevel(TSharedPtr<FGameplayEffectLevelSpec>(new FGameplayEffectLevelSpec(Level, InDef->LevelInfo, Instigator)))
+, Duration(new FAggregator(InDef->Duration.MakeFinalizedCopy(CurveData), ModifierLevel, SKILL_AGG_DEBUG(TEXT("%s Duration"), *InDef->GetName())))
+, Period(new FAggregator(InDef->Period.MakeFinalizedCopy(CurveData), ModifierLevel, SKILL_AGG_DEBUG(TEXT("%s Period"), *InDef->GetName())))
+, ChanceToApplyToTarget(new FAggregator(InDef->ChanceToApplyToTarget.MakeFinalizedCopy(CurveData), ModifierLevel, SKILL_AGG_DEBUG(TEXT("%s ChanceToApplyToTarget"), *InDef->GetName())))
+, ChanceToExecuteOnGameplayEffect(new FAggregator(InDef->ChanceToExecuteOnGameplayEffect.MakeFinalizedCopy(CurveData), ModifierLevel, SKILL_AGG_DEBUG(TEXT("%s ChanceToExecuteOnGameplayEffect"), *InDef->GetName())))
+{
+	Duration.Get()->RegisterLevelDependancies();
+	Period.Get()->RegisterLevelDependancies();
+	ChanceToApplyToTarget.Get()->RegisterLevelDependancies();
+	ChanceToExecuteOnGameplayEffect.Get()->RegisterLevelDependancies();
 
 	InitModifiers(CurveData, Instigator, Level);
 
@@ -377,7 +369,7 @@ bool FGameplayEffectSpec::ShouldApplyAsSnapshot(const FModifierQualifier &Qualif
 
 EGameplayEffectStackingPolicy::Type FGameplayEffectSpec::GetStackingType() const
 {
-	return StackingPolicy;
+	return Def->StackingPolicy;
 }
 
 // --------------------------------------------------------------------------------------------------------------------------------------------------------
@@ -1230,11 +1222,10 @@ void FActiveGameplayEffectsContainer::ExecuteActiveEffectsFrom(const FGameplayEf
 	bool InvokeGameplayCueExecute = false;
 
 	// check if this is a stacking effect and if it is the active stacking effect.
-	// todo: move this into the loop if we end up applying stacking to different modifiers on the same effect
 	if (Spec.GetStackingType() != EGameplayEffectStackingPolicy::Unlimited)
 	{
 		// we're a stacking attribute but not active
-		if (Spec.StackedAttribName == NAME_None)
+		if (Spec.bTopOfStack == false)
 		{
 			return;
 		}
@@ -1475,7 +1466,7 @@ bool FActiveGameplayEffectsContainer::IsGameplayEffectActive(FActiveGameplayEffe
 		{
 			// stacking effects may return false
 			if (Effect.Spec.GetStackingType() != EGameplayEffectStackingPolicy::Unlimited &&
-				Effect.Spec.StackedAttribName == NAME_None)
+				Effect.Spec.bTopOfStack == false)
 			{
 				return false;
 			}
@@ -1740,7 +1731,7 @@ void FActiveGameplayEffectsContainer::RecalculateStacking()
 		{
 			if (Effect.Spec.Def->Modifiers.Num() > 0)
 			{
-				Effect.Spec.StackedAttribName = NAME_None;
+				Effect.Spec.bTopOfStack = false;
 				bool bFoundStack = false;
 
 				// group all of the custom stacking effects and deal with them after
@@ -1825,7 +1816,7 @@ void FActiveGameplayEffectsContainer::RecalculateStacking()
 	// we've found all of the best stacking effects, mark them so that they updated properly
 	for (FActiveGameplayEffect* const StackedEffect : StackedEffects)
 	{
-		StackedEffect->Spec.StackedAttribName = FName(*StackedEffect->Spec.Def->Modifiers[0].Attribute.GetName());
+		StackedEffect->Spec.bTopOfStack = true;
 	}
 }
 
diff --git a/Engine/Source/Runtime/SkillSystem/Private/GameplayEffectTests.cpp b/Engine/Source/Runtime/SkillSystem/Private/GameplayEffectTests.cpp
--- a/Engine/Source/Runtime/SkillSystem/Private/GameplayEffectTests.cpp
+++ b/Engine/Source/Runtime/SkillSystem/Private/GameplayEffectTests.cpp
@@ -3369,6 +3369,9 @@ bool GameplayEffectsTest_StackingHighest(UWorld *World, FAutomationTestBase * Te
 		BaseStackedEffect->Duration.Value = UGameplayEffect::INFINITE_DURATION;
 		BaseStackedEffect->Period.Value = 1.f;
 		BaseStackedEffect->StackingPolicy = EGameplayEffectStackingPolicy::Highest;
+		BaseStackedEffect->StackedAttribName = FName(*StackingProperty->GetName());
+
+		BaseStackedEffect->ValidateGameplayEffect();
 
 		SourceComponent->ApplyGameplayEffectToTarget(BaseStackedEffect, DestComponent);
 
@@ -3402,6 +3405,9 @@ bool GameplayEffectsTest_StackingHighest(UWorld *World, FAutomationTestBase * Te
 		BaseStackedEffect->Duration.Value = UGameplayEffect::INFINITE_DURATION;
 		BaseStackedEffect->Period.Value = 1.f;
 		BaseStackedEffect->StackingPolicy = EGameplayEffectStackingPolicy::Highest;
+		BaseStackedEffect->StackedAttribName = FName(*StackingProperty->GetName());
+
+		BaseStackedEffect->ValidateGameplayEffect();
 
 		SourceComponent->ApplyGameplayEffectToTarget(BaseStackedEffect, DestComponent);
 
@@ -3455,6 +3461,9 @@ bool GameplayEffectsTest_StackingLowest(UWorld *World, FAutomationTestBase * Tes
 		BaseStackedEffect->Duration.Value = UGameplayEffect::INFINITE_DURATION;
 		BaseStackedEffect->Period.Value = 1.f;
 		BaseStackedEffect->StackingPolicy = EGameplayEffectStackingPolicy::Lowest;
+		BaseStackedEffect->StackedAttribName = FName(*StackingProperty->GetName());
+
+		BaseStackedEffect->ValidateGameplayEffect();
 
 		SourceComponent->ApplyGameplayEffectToTarget(BaseStackedEffect, DestComponent);
 
@@ -3483,6 +3492,9 @@ bool GameplayEffectsTest_StackingLowest(UWorld *World, FAutomationTestBase * Tes
 		BaseStackedEffect->Duration.Value = UGameplayEffect::INFINITE_DURATION;
 		BaseStackedEffect->Period.Value = 1.f;
 		BaseStackedEffect->StackingPolicy = EGameplayEffectStackingPolicy::Lowest;
+		BaseStackedEffect->StackedAttribName = FName(*StackingProperty->GetName());
+
+		BaseStackedEffect->ValidateGameplayEffect();
 
 		SourceComponent->ApplyGameplayEffectToTarget(BaseStackedEffect, DestComponent);
 
@@ -3536,6 +3548,9 @@ bool GameplayEffectsTest_StackingUnlimited(UWorld *World, FAutomationTestBase *
 		BaseStackedEffect->Duration.Value = UGameplayEffect::INFINITE_DURATION;
 		BaseStackedEffect->Period.Value = 1.f;
 		BaseStackedEffect->StackingPolicy = EGameplayEffectStackingPolicy::Unlimited;
+		BaseStackedEffect->StackedAttribName = FName(*StackingProperty->GetName());
+
+		BaseStackedEffect->ValidateGameplayEffect();
 
 		SourceComponent->ApplyGameplayEffectToTarget(BaseStackedEffect, DestComponent);
 	}
@@ -3552,6 +3567,9 @@ bool GameplayEffectsTest_StackingUnlimited(UWorld *World, FAutomationTestBase *
 		BaseStackedEffect->Duration.Value = UGameplayEffect::INFINITE_DURATION;
 		BaseStackedEffect->Period.Value = 1.f;
 		BaseStackedEffect->StackingPolicy = EGameplayEffectStackingPolicy::Unlimited;
+		BaseStackedEffect->StackedAttribName = FName(*StackingProperty->GetName());
+
+		BaseStackedEffect->ValidateGameplayEffect();
 
 		SourceComponent->ApplyGameplayEffectToTarget(BaseStackedEffect, DestComponent);
 
@@ -3605,6 +3623,9 @@ bool GameplayEffectsTest_StackingRemoval(UWorld *World, FAutomationTestBase * Te
 		BaseStackedEffect->Duration.Value = UGameplayEffect::INFINITE_DURATION;
 		BaseStackedEffect->Period.Value = 1.f;
 		BaseStackedEffect->StackingPolicy = EGameplayEffectStackingPolicy::Highest;
+		BaseStackedEffect->StackedAttribName = FName(*StackingProperty->GetName());
+
+		BaseStackedEffect->ValidateGameplayEffect();
 
 		SourceComponent->ApplyGameplayEffectToTarget(BaseStackedEffect, DestComponent);
 	}
@@ -3621,6 +3642,9 @@ bool GameplayEffectsTest_StackingRemoval(UWorld *World, FAutomationTestBase * Te
 		BaseStackedEffect->Duration.Value = 1.f;// UGameplayEffect::INFINITE_DURATION;
 		BaseStackedEffect->Period.Value = 1.f;
 		BaseStackedEffect->StackingPolicy = EGameplayEffectStackingPolicy::Highest;
+		BaseStackedEffect->StackedAttribName = FName(*StackingProperty->GetName());
+
+		BaseStackedEffect->ValidateGameplayEffect();
 
 		SourceComponent->ApplyGameplayEffectToTarget(BaseStackedEffect, DestComponent);
 
@@ -3681,6 +3705,9 @@ bool GameplayEffectsTest_StackingNoStack(UWorld *World, FAutomationTestBase * Te
 		UnstackableEffect->Duration.Value = UGameplayEffect::INFINITE_DURATION;
 		UnstackableEffect->Period.Value = 1.f;
 		UnstackableEffect->StackingPolicy = EGameplayEffectStackingPolicy::Highest;
+		UnstackableEffect->StackedAttribName = FName(*NoStackProperty->GetName());
+
+		UnstackableEffect->ValidateGameplayEffect();
 
 		SourceComponent->ApplyGameplayEffectToTarget(UnstackableEffect, DestComponent);
 	}
@@ -3697,6 +3724,9 @@ bool GameplayEffectsTest_StackingNoStack(UWorld *World, FAutomationTestBase * Te
 		UnstackableEffect->Duration.Value = UGameplayEffect::INFINITE_DURATION;
 		UnstackableEffect->Period.Value = 1.f;
 		UnstackableEffect->StackingPolicy = EGameplayEffectStackingPolicy::Highest;
+		UnstackableEffect->StackedAttribName = FName(*NoStackProperty->GetName());
+
+		UnstackableEffect->ValidateGameplayEffect();
 
 		SourceComponent->ApplyGameplayEffectToTarget(UnstackableEffect, DestComponent);
 
@@ -3752,6 +3782,9 @@ bool GameplayEffectsTest_StackingCustomCapped(UWorld *World, FAutomationTestBase
 		BaseStackedEffect->Period.Value = 1.f;
 		BaseStackedEffect->StackingPolicy = EGameplayEffectStackingPolicy::Callback;
 		BaseStackedEffect->StackingExtension = UGameplayEffectStackingExtension_CappedNumberTest::StaticClass();
+		BaseStackedEffect->StackedAttribName = FName(*StackingProperty->GetName());
+
+		BaseStackedEffect->ValidateGameplayEffect();
 
 		SourceComponent->ApplyGameplayEffectToTarget(BaseStackedEffect, DestComponent);
 	}
@@ -3770,6 +3803,9 @@ bool GameplayEffectsTest_StackingCustomCapped(UWorld *World, FAutomationTestBase
 		BaseStackedEffect->Period.Value = 1.f;
 		BaseStackedEffect->StackingPolicy = EGameplayEffectStackingPolicy::Callback;
 		BaseStackedEffect->StackingExtension = UGameplayEffectStackingExtension_CappedNumberTest::StaticClass();
+		BaseStackedEffect->StackedAttribName = FName(*StackingProperty->GetName());
+
+		BaseStackedEffect->ValidateGameplayEffect();
 
 		SourceComponent->ApplyGameplayEffectToTarget(BaseStackedEffect, DestComponent);
 	}
@@ -3788,6 +3824,9 @@ bool GameplayEffectsTest_StackingCustomCapped(UWorld *World, FAutomationTestBase
 		BaseStackedEffect->Period.Value = 1.f;
 		BaseStackedEffect->StackingPolicy = EGameplayEffectStackingPolicy::Callback;
 		BaseStackedEffect->StackingExtension = UGameplayEffectStackingExtension_CappedNumberTest::StaticClass();
+		BaseStackedEffect->StackedAttribName = FName(*StackingProperty->GetName());
+
+		BaseStackedEffect->ValidateGameplayEffect();
 
 		SourceComponent->ApplyGameplayEffectToTarget(BaseStackedEffect, DestComponent);
 
@@ -3818,6 +3857,9 @@ bool GameplayEffectsTest_StackingCustomCapped(UWorld *World, FAutomationTestBase
 		BaseStackedEffect->Period.Value = 1.f;
 		BaseStackedEffect->StackingPolicy = EGameplayEffectStackingPolicy::Callback;
 		BaseStackedEffect->StackingExtension = UGameplayEffectStackingExtension_CappedNumberTest::StaticClass();
+		BaseStackedEffect->StackedAttribName = FName(*StackingProperty->GetName());
+
+		BaseStackedEffect->ValidateGameplayEffect();
 
 		SourceComponent->ApplyGameplayEffectToTarget(BaseStackedEffect, DestComponent);
 
@@ -3891,6 +3933,9 @@ bool GameplayEffectsTest_StackingCustomDiminishingReturns(UWorld *World, FAutoma
 		BaseStackedEffect->Period.Value = 1.f;
 		BaseStackedEffect->StackingPolicy = EGameplayEffectStackingPolicy::Callback;
 		BaseStackedEffect->StackingExtension = UGameplayEffectStackingExtension_DiminishingReturnsTest::StaticClass();
+		BaseStackedEffect->StackedAttribName = FName(*StackingProperty->GetName());
+
+		BaseStackedEffect->ValidateGameplayEffect();
 
 		SourceComponent->ApplyGameplayEffectToTarget(BaseStackedEffect, DestComponent);
 
@@ -3921,6 +3966,9 @@ bool GameplayEffectsTest_StackingCustomDiminishingReturns(UWorld *World, FAutoma
 		BaseStackedEffect->Period.Value = 1.f;
 		BaseStackedEffect->StackingPolicy = EGameplayEffectStackingPolicy::Callback;
 		BaseStackedEffect->StackingExtension = UGameplayEffectStackingExtension_DiminishingReturnsTest::StaticClass();
+		BaseStackedEffect->StackedAttribName = FName(*StackingProperty->GetName());
+
+		BaseStackedEffect->ValidateGameplayEffect();
 
 		SourceComponent->ApplyGameplayEffectToTarget(BaseStackedEffect, DestComponent);
 
@@ -3951,6 +3999,9 @@ bool GameplayEffectsTest_StackingCustomDiminishingReturns(UWorld *World, FAutoma
 		BaseStackedEffect->Period.Value = 1.f;
 		BaseStackedEffect->StackingPolicy = EGameplayEffectStackingPolicy::Callback;
 		BaseStackedEffect->StackingExtension = UGameplayEffectStackingExtension_DiminishingReturnsTest::StaticClass();
+		BaseStackedEffect->StackedAttribName = FName(*StackingProperty->GetName());
+
+		BaseStackedEffect->ValidateGameplayEffect();
 
 		SourceComponent->ApplyGameplayEffectToTarget(BaseStackedEffect, DestComponent);
 
@@ -3981,6 +4032,9 @@ bool GameplayEffectsTest_StackingCustomDiminishingReturns(UWorld *World, FAutoma
 		BaseStackedEffect->Period.Value = 1.f;
 		BaseStackedEffect->StackingPolicy = EGameplayEffectStackingPolicy::Callback;
 		BaseStackedEffect->StackingExtension = UGameplayEffectStackingExtension_DiminishingReturnsTest::StaticClass();
+		BaseStackedEffect->StackedAttribName = FName(*StackingProperty->GetName());
+
+		BaseStackedEffect->ValidateGameplayEffect();
 
 		SourceComponent->ApplyGameplayEffectToTarget(BaseStackedEffect, DestComponent);
 
@@ -4034,6 +4088,9 @@ bool GameplayEffectsTest_StackingDifferentRules(UWorld *World, FAutomationTestBa
 		BaseStackedEffect->Duration.Value = UGameplayEffect::INFINITE_DURATION;
 		BaseStackedEffect->Period.Value = 1.f;
 		BaseStackedEffect->StackingPolicy = EGameplayEffectStackingPolicy::Highest;
+		BaseStackedEffect->StackedAttribName = FName(*StackingProperty->GetName());
+
+		BaseStackedEffect->ValidateGameplayEffect();
 
 		SourceComponent->ApplyGameplayEffectToTarget(BaseStackedEffect, DestComponent);
 	}
@@ -4050,6 +4107,9 @@ bool GameplayEffectsTest_StackingDifferentRules(UWorld *World, FAutomationTestBa
 		BaseStackedEffect->Duration.Value = UGameplayEffect::INFINITE_DURATION;
 		BaseStackedEffect->Period.Value = 1.f;
 		BaseStackedEffect->StackingPolicy = EGameplayEffectStackingPolicy::Lowest;
+		BaseStackedEffect->StackedAttribName = FName(*StackingProperty->GetName());
+
+		BaseStackedEffect->ValidateGameplayEffect();
 
 		SourceComponent->ApplyGameplayEffectToTarget(BaseStackedEffect, DestComponent);
 
@@ -4104,6 +4164,9 @@ bool GameplayEffectsTest_StackingDifferentAttributes(UWorld *World, FAutomationT
 		BaseStackedEffect->Duration.Value = UGameplayEffect::INFINITE_DURATION;
 		BaseStackedEffect->Period.Value = 1.f;
 		BaseStackedEffect->StackingPolicy = EGameplayEffectStackingPolicy::Highest;
+		BaseStackedEffect->StackedAttribName = FName(*StackingProperty1->GetName());
+
+		BaseStackedEffect->ValidateGameplayEffect();
 
 		SourceComponent->ApplyGameplayEffectToTarget(BaseStackedEffect, DestComponent);
 	}
@@ -4120,6 +4183,9 @@ bool GameplayEffectsTest_StackingDifferentAttributes(UWorld *World, FAutomationT
 		BaseStackedEffect->Duration.Value = UGameplayEffect::INFINITE_DURATION;
 		BaseStackedEffect->Period.Value = 1.f;
 		BaseStackedEffect->StackingPolicy = EGameplayEffectStackingPolicy::Highest;
+		BaseStackedEffect->StackedAttribName = FName(*StackingProperty2->GetName());
+
+		BaseStackedEffect->ValidateGameplayEffect();
 
 		SourceComponent->ApplyGameplayEffectToTarget(BaseStackedEffect, DestComponent);
 
@@ -4178,6 +4244,9 @@ bool GameplayEffectsTest_StackingCustomTwoRules(UWorld *World, FAutomationTestBa
 		BaseStackedEffect->Period.Value = 1.f;
 		BaseStackedEffect->StackingPolicy = EGameplayEffectStackingPolicy::Callback;
 		BaseStackedEffect->StackingExtension = UGameplayEffectStackingExtension_CappedNumberTest::StaticClass();
+		BaseStackedEffect->StackedAttribName = FName(*StackingProperty->GetName());
+
+		BaseStackedEffect->ValidateGameplayEffect();
 
 		SourceComponent->ApplyGameplayEffectToTarget(BaseStackedEffect, DestComponent);
 	}
@@ -4196,6 +4265,9 @@ bool GameplayEffectsTest_StackingCustomTwoRules(UWorld *World, FAutomationTestBa
 		BaseStackedEffect->Period.Value = 1.f;
 		BaseStackedEffect->StackingPolicy = EGameplayEffectStackingPolicy::Callback;
 		BaseStackedEffect->StackingExtension = UGameplayEffectStackingExtension_DiminishingReturnsTest::StaticClass();
+		BaseStackedEffect->StackedAttribName = FName(*StackingProperty->GetName());
+
+		BaseStackedEffect->ValidateGameplayEffect();
 
 		SourceComponent->ApplyGameplayEffectToTarget(BaseStackedEffect, DestComponent);
 	}
@@ -4214,6 +4286,9 @@ bool GameplayEffectsTest_StackingCustomTwoRules(UWorld *World, FAutomationTestBa
 		BaseStackedEffect->Period.Value = 1.f;
 		BaseStackedEffect->StackingPolicy = EGameplayEffectStackingPolicy::Callback;
 		BaseStackedEffect->StackingExtension = UGameplayEffectStackingExtension_DiminishingReturnsTest::StaticClass();
+		BaseStackedEffect->StackedAttribName = FName(*StackingProperty->GetName());
+
+		BaseStackedEffect->ValidateGameplayEffect();
 
 		SourceComponent->ApplyGameplayEffectToTarget(BaseStackedEffect, DestComponent);
 
@@ -4270,6 +4345,9 @@ bool GameplayEffectsTest_StackingCustomTwoAttributes(UWorld *World, FAutomationT
 		BaseStackedEffect->Period.Value = 1.f;
 		BaseStackedEffect->StackingPolicy = EGameplayEffectStackingPolicy::Callback;
 		BaseStackedEffect->StackingExtension = UGameplayEffectStackingExtension_CappedNumberTest::StaticClass();
+		BaseStackedEffect->StackedAttribName = FName(*StackingProperty1->GetName());
+
+		BaseStackedEffect->ValidateGameplayEffect();
 
 		SourceComponent->ApplyGameplayEffectToTarget(BaseStackedEffect, DestComponent);
 	}
@@ -4288,6 +4366,9 @@ bool GameplayEffectsTest_StackingCustomTwoAttributes(UWorld *World, FAutomationT
 		BaseStackedEffect->Period.Value = 1.f;
 		BaseStackedEffect->StackingPolicy = EGameplayEffectStackingPolicy::Callback;
 		BaseStackedEffect->StackingExtension = UGameplayEffectStackingExtension_CappedNumberTest::StaticClass();
+		BaseStackedEffect->StackedAttribName = FName(*StackingProperty1->GetName());
+
+		BaseStackedEffect->ValidateGameplayEffect();
 
 		SourceComponent->ApplyGameplayEffectToTarget(BaseStackedEffect, DestComponent);
 	}
@@ -4306,6 +4387,9 @@ bool GameplayEffectsTest_StackingCustomTwoAttributes(UWorld *World, FAutomationT
 		BaseStackedEffect->Period.Value = 1.f;
 		BaseStackedEffect->StackingPolicy = EGameplayEffectStackingPolicy::Callback;
 		BaseStackedEffect->StackingExtension = UGameplayEffectStackingExtension_CappedNumberTest::StaticClass();
+		BaseStackedEffect->StackedAttribName = FName(*StackingProperty2->GetName());
+
+		BaseStackedEffect->ValidateGameplayEffect();
 
 		SourceComponent->ApplyGameplayEffectToTarget(BaseStackedEffect, DestComponent);
 
@@ -4382,6 +4466,9 @@ bool GameplayEffectsTest_StackingRemovingModifiers(UWorld *World, FAutomationTes
 		BaseStackedEffect->Period.Value = 1.f;
 		BaseStackedEffect->StackingPolicy = EGameplayEffectStackingPolicy::Highest;
 		BaseStackedEffect->GameplayEffectTags.AddTag(IGameplayTagsModule::RequestGameplayTag(FName(TEXT("Stack"))));
+		BaseStackedEffect->StackedAttribName = FName(*StackingProperty->GetName());
+
+		BaseStackedEffect->ValidateGameplayEffect();
 
 		SourceComponent->ApplyGameplayEffectToTarget(BaseStackedEffect, DestComponent);
 	}
@@ -4398,6 +4485,9 @@ bool GameplayEffectsTest_StackingRemovingModifiers(UWorld *World, FAutomationTes
 		BaseStackedEffect->Duration.Value = UGameplayEffect::INFINITE_DURATION;
 		BaseStackedEffect->Period.Value = 1.f;
 		BaseStackedEffect->StackingPolicy = EGameplayEffectStackingPolicy::Highest;
+		BaseStackedEffect->StackedAttribName = FName(*StackingProperty->GetName());
+
+		BaseStackedEffect->ValidateGameplayEffect();
 
 		SourceComponent->ApplyGameplayEffectToTarget(BaseStackedEffect, DestComponent);
 
@@ -4460,6 +4550,9 @@ bool GameplayEffectsTest_StackingAddingModifiers(UWorld *World, FAutomationTestB
 		BaseStackedEffect->Period.Value = 1.f;
 		BaseStackedEffect->StackingPolicy = EGameplayEffectStackingPolicy::Highest;
 		BaseStackedEffect->GameplayEffectTags.AddTag(IGameplayTagsModule::RequestGameplayTag(FName(TEXT("Stack"))));
+		BaseStackedEffect->StackedAttribName = FName(*StackingProperty->GetName());
+
+		BaseStackedEffect->ValidateGameplayEffect();
 
 		SourceComponent->ApplyGameplayEffectToTarget(BaseStackedEffect, DestComponent);
 	}
@@ -4476,6 +4569,9 @@ bool GameplayEffectsTest_StackingAddingModifiers(UWorld *World, FAutomationTestB
 		BaseStackedEffect->Duration.Value = UGameplayEffect::INFINITE_DURATION;
 		BaseStackedEffect->Period.Value = 1.f;
 		BaseStackedEffect->StackingPolicy = EGameplayEffectStackingPolicy::Highest;
+		BaseStackedEffect->StackedAttribName = FName(*StackingProperty->GetName());
+
+		BaseStackedEffect->ValidateGameplayEffect();
 
 		SourceComponent->ApplyGameplayEffectToTarget(BaseStackedEffect, DestComponent);
 
