CIS fix

diff --git a/Engine/Source/Runtime/SkillSystem/Public/SkillSystemTypes.h b/Engine/Source/Runtime/SkillSystem/Public/SkillSystemTypes.h
--- a/Engine/Source/Runtime/SkillSystem/Public/SkillSystemTypes.h
+++ b/Engine/Source/Runtime/SkillSystem/Public/SkillSystemTypes.h
@@ -135,6 +135,8 @@ struct SKILLSYSTEM_API FGameplayAbilityActivationInfo
 
 struct SKILLSYSTEM_API FGameplayAbilityTargetData
 {
+	virtual ~FGameplayAbilityTargetData() { }
+
 	void ApplyGameplayEffect(UGameplayEffect *GameplayEffect, const FGameplayAbilityActorInfo InstigatorInfo);
 
 	virtual TArray<AActor*>	GetActors() const = 0;
