PS4 - Fix compiler errors

diff --git a/Engine/Source/Runtime/SkillSystem/Public/SkillSystemTypes.h b/Engine/Source/Runtime/SkillSystem/Public/SkillSystemTypes.h
--- a/Engine/Source/Runtime/SkillSystem/Public/SkillSystemTypes.h
+++ b/Engine/Source/Runtime/SkillSystem/Public/SkillSystemTypes.h
@@ -35,6 +35,8 @@ struct SKILLSYSTEM_API FGameplayAbilityActorInfo
 {
 	GENERATED_USTRUCT_BODY()
 
+	virtual ~FGameplayAbilityActorInfo() {}
+
 	UPROPERTY(BlueprintReadOnly, Category = "ActorInfo")
 	TWeakObjectPtr<AActor>	Actor;
 
@@ -66,8 +68,8 @@ struct SKILLSYSTEM_API FGameplayAbilityActivationInfo
 	GENERATED_USTRUCT_BODY()
 
 	FGameplayAbilityActivationInfo()
-		: PredictionKey(0)
-		, ActivationMode(EGameplayAbilityActivationMode::Authority)
+		: ActivationMode(EGameplayAbilityActivationMode::Authority),
+		PredictionKey(0)
 	{
 
 	}
